#ttp 330142 - EDITOR: Users are unable to see the file path of a project in the project
#proj UE4
#branch UE4
#summary Added a tooltip to show project information, as well as a context menu option to go to the .uproject file in explorer
#extra There is now a tooltip (styled like the content browser tooltip) showing the description (if available) and the path to the project. There is also a context menu option to go to the .uproject file for the selected project.

diff --git a/Engine/Source/Editor/ContentBrowser/Private/ContentBrowserUtils.cpp b/Engine/Source/Editor/ContentBrowser/Private/ContentBrowserUtils.cpp
--- a/Engine/Source/Editor/ContentBrowser/Private/ContentBrowserUtils.cpp
+++ b/Engine/Source/Editor/ContentBrowser/Private/ContentBrowserUtils.cpp
@@ -1225,11 +1225,9 @@ FLinearColor ContentBrowserUtils::GetDefaultColor()
 
 FText ContentBrowserUtils::GetExploreFolderText()
 {
-#if PLATFORM_MAC
-	return LOCTEXT("Mac_ShowInFinder", "Show In Finder");
-#else
-	return LOCTEXT("Win_ShowInExplorer", "Show In Explorer");
-#endif
+	FFormatNamedArguments Args;
+	Args.Add(TEXT("FileManagerName"), FPlatformMisc::GetFileManagerName());
+	return FText::Format(NSLOCTEXT("GenericPlatform", "ShowInFileManager", "Show In {FileManagerName}"), Args);
 }
 
 #undef LOCTEXT_NAMESPACE
diff --git a/Engine/Source/Editor/EditorStyle/Private/SlateEditorStyle.cpp b/Engine/Source/Editor/EditorStyle/Private/SlateEditorStyle.cpp
--- a/Engine/Source/Editor/EditorStyle/Private/SlateEditorStyle.cpp
+++ b/Engine/Source/Editor/EditorStyle/Private/SlateEditorStyle.cpp
@@ -4439,6 +4439,11 @@ void FSlateEditorStyle::FStyle::SetupToolkitStyles()
 		Set("ProjectBrowser.Tab.ActiveHighlight", new BOX_BRUSH("Common/FlatColorSquare", FVector2D(1.0f, 1.0f), FMargin(0), SelectionColor) );
 		Set("ProjectBrowser.Tab.Highlight", new BOX_BRUSH("Common/FlatColorSquare", FVector2D(1.0f, 1.0f), FMargin(0), SelectionColor_Inactive));
 		Set("ProjectBrowser.Tab.PressedHighlight", new BOX_BRUSH("Common/FlatColorSquare", FVector2D(1.0f, 1.0f), FMargin(0), SelectionColor_Pressed));
+
+		Set( "ProjectBrowser.TileViewTooltip.ToolTipBorder", new FSlateColorBrush( FLinearColor::Black ) );
+		Set( "ProjectBrowser.TileViewTooltip.NonContentBorder", new BOX_BRUSH( "/Docking/TabContentArea", FMargin(4/16.0f) ) );
+		Set( "ProjectBrowser.TileViewTooltip.ContentBorder", new BOX_BRUSH( "Common/GroupBorder", FMargin(4.0f/16.0f) ) );
+		Set( "ProjectBrowser.TileViewTooltip.NameFont", TTF_CORE_FONT( "Fonts/Roboto-Regular", 12 ) );
 	}
 
 	// Toolkit Display
diff --git a/Engine/Source/Editor/GameProjectGeneration/Private/SProjectBrowser.cpp b/Engine/Source/Editor/GameProjectGeneration/Private/SProjectBrowser.cpp
--- a/Engine/Source/Editor/GameProjectGeneration/Private/SProjectBrowser.cpp
+++ b/Engine/Source/Editor/GameProjectGeneration/Private/SProjectBrowser.cpp
@@ -254,6 +254,7 @@ void SProjectBrowser::ConstructCategory( const TSharedRef<SVerticalBox>& InCateg
 		.SelectionMode(ESelectionMode::Single)
 		.ClearSelectionOnClick(false)
 		.OnGenerateTile(this, &SProjectBrowser::MakeProjectViewWidget)
+		.OnContextMenuOpening(this, &SProjectBrowser::OnGetContextMenuContent)
 		.OnMouseButtonDoubleClick(this, &SProjectBrowser::HandleProjectItemDoubleClick)
 		.OnSelectionChanged(TSlateDelegates<TSharedPtr<FProjectItem>>::FOnSelectionChanged::CreateSP(this, &SProjectBrowser::HandleProjectViewSelectionChanged, Category->CategoryName))
 		.ItemHeight(ThumbnailSize + ThumbnailBorderPadding + 32)
@@ -309,37 +310,173 @@ TSharedRef<ITableRow> SProjectBrowser::MakeProjectViewWidget(TSharedPtr<FProject
 			];
 	}
 
-	return
-		SNew( STableRow<TSharedPtr<FProjectItem>>, OwnerTable )
-		.Style(FEditorStyle::Get(), "GameProjectDialog.TemplateListView.TableRow")
+	TSharedRef<ITableRow> TableRow = SNew( STableRow<TSharedPtr<FProjectItem>>, OwnerTable )
+	.Style(FEditorStyle::Get(), "GameProjectDialog.TemplateListView.TableRow")
+	[
+		SNew(SBox)
+		.HeightOverride(ThumbnailSize+ThumbnailBorderPadding+5)
 		[
-			SNew(SBox).HeightOverride(ThumbnailSize+ThumbnailBorderPadding+5)
+			SNew(SVerticalBox)
+
+			// Thumbnail
+			+SVerticalBox::Slot()
+			.AutoHeight()
 			[
-				SNew(SVerticalBox)
+				SNew(SBox)
+				.WidthOverride(ThumbnailSize + ThumbnailBorderPadding * 2)
+				.HeightOverride(ThumbnailSize + ThumbnailBorderPadding * 2)
+				[
+					Thumbnail.ToSharedRef()
+				]
+			]
 
-				// Thumbnail
-				+SVerticalBox::Slot()
-				.AutoHeight()
+			// Name
+			+SVerticalBox::Slot()
+			.HAlign(HAlign_Center)
+			.VAlign(VAlign_Top)
+			[
+				SNew(STextBlock)
+				.HighlightText(this, &SProjectBrowser::GetItemHighlightText)
+				.Text(ProjectItem->Name)
+			]
+		]
+	];
+
+	TableRow->AsWidget()->SetToolTip(MakeProjectToolTip(ProjectItem));
+
+	return TableRow;
+}
+
+
+TSharedRef<SToolTip> SProjectBrowser::MakeProjectToolTip( TSharedPtr<FProjectItem> ProjectItem ) const
+{
+	// Create a box to hold every line of info in the body of the tooltip
+	TSharedRef<SVerticalBox> InfoBox = SNew(SVerticalBox);
+
+	if(!ProjectItem->Description.IsEmpty())
+	{
+		AddToToolTipInfoBox( InfoBox, LOCTEXT("ProjectTileTooltipDescription", "Description"), ProjectItem->Description );
+	}
+
+	{
+		const FString ProjectPath = FPaths::GetPath(ProjectItem->ProjectFile);
+		AddToToolTipInfoBox( InfoBox, LOCTEXT("ProjectTileTooltipPath", "Path"), FText::FromString(ProjectPath) );
+	}
+
+	TSharedRef<SToolTip> Tooltip = SNew(SToolTip)
+	.TextMargin(1)
+	.BorderImage( FEditorStyle::GetBrush("ProjectBrowser.TileViewTooltip.ToolTipBorder") )
+	[
+		SNew(SBorder)
+		.Padding(6)
+		.BorderImage( FEditorStyle::GetBrush("ProjectBrowser.TileViewTooltip.NonContentBorder") )
+		[
+			SNew(SVerticalBox)
+
+			+SVerticalBox::Slot()
+			.AutoHeight()
+			.Padding(0, 0, 0, 4)
+			[
+				SNew(SBorder)
+				.Padding(6)
+				.BorderImage( FEditorStyle::GetBrush("ProjectBrowser.TileViewTooltip.ContentBorder") )
 				[
-					SNew(SBox)
-					.WidthOverride(ThumbnailSize + ThumbnailBorderPadding * 2)
-					.HeightOverride(ThumbnailSize + ThumbnailBorderPadding * 2)
+					SNew(SVerticalBox)
+
+					+SVerticalBox::Slot()
+					.AutoHeight()
+					.VAlign(VAlign_Center)
 					[
-						Thumbnail.ToSharedRef()
+						SNew(STextBlock)
+						.Text( ProjectItem->Name )
+						.Font( FEditorStyle::GetFontStyle("ProjectBrowser.TileViewTooltip.NameFont") )
 					]
 				]
+			]
 
-				// Name
-				+SVerticalBox::Slot()
-				.HAlign(HAlign_Center)
-				.VAlign(VAlign_Top)
+			+SVerticalBox::Slot()
+			.AutoHeight()
+			[
+				SNew(SBorder)
+				.Padding(6)
+				.BorderImage( FEditorStyle::GetBrush("ProjectBrowser.TileViewTooltip.ContentBorder") )
 				[
-					SNew(STextBlock)
-					.HighlightText(this, &SProjectBrowser::GetItemHighlightText)
-					.Text(ProjectItem->Name)
+					InfoBox
 				]
 			]
-		];
+		]
+	];
+
+	return Tooltip;
+}
+
+
+void SProjectBrowser::AddToToolTipInfoBox(const TSharedRef<SVerticalBox>& InfoBox, const FText& Key, const FText& Value) const
+{
+	InfoBox->AddSlot()
+	.AutoHeight()
+	.Padding(0, 1)
+	[
+		SNew(SHorizontalBox)
+
+		+SHorizontalBox::Slot()
+		.AutoWidth()
+		.Padding(0, 0, 4, 0)
+		[
+			SNew(STextBlock) .Text( FText::Format(LOCTEXT("ProjectBrowserTooltipFormat", "{0}:"), Key ) )
+			.ColorAndOpacity(FSlateColor::UseSubduedForeground())
+		]
+
+		+SHorizontalBox::Slot()
+		.AutoWidth()
+		[
+			SNew(STextBlock) .Text( Value )
+			.ColorAndOpacity(FSlateColor::UseForeground())
+		]
+	];
+}
+
+
+TSharedPtr<SWidget> SProjectBrowser::OnGetContextMenuContent() const
+{
+	FMenuBuilder MenuBuilder(/*bInShouldCloseWindowAfterMenuSelection=*/true, nullptr);
+
+	TSharedPtr<FProjectItem> SelectedProjectItem = GetSelectedProjectItem();
+	const FText ProjectContextActionsText = (SelectedProjectItem.IsValid()) ? SelectedProjectItem->Name : LOCTEXT("ProjectActionsMenuHeading", "Project Actions");
+	MenuBuilder.BeginSection("ProjectContextActions", ProjectContextActionsText);
+
+	FFormatNamedArguments Args;
+	Args.Add(TEXT("FileManagerName"), FPlatformMisc::GetFileManagerName());
+	const FText ExploreToText = FText::Format(NSLOCTEXT("GenericPlatform", "ShowInFileManager", "Show In {FileManagerName}"), Args);
+
+	MenuBuilder.AddMenuEntry(
+		ExploreToText,
+		LOCTEXT("FindInExplorerTooltip", "Finds this project on disk"),
+		FSlateIcon(),
+		FUIAction(
+		FExecuteAction::CreateSP( this, &SProjectBrowser::ExecuteFindInExplorer ),
+		FCanExecuteAction::CreateSP( this, &SProjectBrowser::CanExecuteFindInExplorer )
+		)
+		);
+
+	MenuBuilder.EndSection();
+
+	return MenuBuilder.MakeWidget();
+}
+
+
+void SProjectBrowser::ExecuteFindInExplorer() const
+{
+	TSharedPtr<FProjectItem> SelectedProjectItem = GetSelectedProjectItem();
+	check(SelectedProjectItem.IsValid());
+	FPlatformProcess::ExploreFolder(*SelectedProjectItem->ProjectFile);
+}
+
+
+bool SProjectBrowser::CanExecuteFindInExplorer() const
+{
+	TSharedPtr<FProjectItem> SelectedProjectItem = GetSelectedProjectItem();
+	return SelectedProjectItem.IsValid();
 }
 
 
diff --git a/Engine/Source/Editor/GameProjectGeneration/Private/SProjectBrowser.h b/Engine/Source/Editor/GameProjectGeneration/Private/SProjectBrowser.h
--- a/Engine/Source/Editor/GameProjectGeneration/Private/SProjectBrowser.h
+++ b/Engine/Source/Editor/GameProjectGeneration/Private/SProjectBrowser.h
@@ -44,6 +44,21 @@ protected:
 	/** Creates a row in the template list */
 	TSharedRef<ITableRow> MakeProjectViewWidget( TSharedPtr<FProjectItem> ProjectItem, const TSharedRef<STableViewBase>& OwnerTable );
 
+	/** Create a tooltip for the given project item */
+	TSharedRef<SToolTip> MakeProjectToolTip( TSharedPtr<FProjectItem> ProjectItem ) const;
+
+	/** Add information to the tooltip for this project item */
+	void AddToToolTipInfoBox(const TSharedRef<SVerticalBox>& InfoBox, const FText& Key, const FText& Value) const;
+
+	/** Get the context menu to use for the selected project item */
+	TSharedPtr<SWidget> OnGetContextMenuContent() const;
+
+	/** Handler for when find in explorer is selected */
+	void ExecuteFindInExplorer() const;
+
+	/** Handler to check to see if a find in explorer command is allowed */
+	bool CanExecuteFindInExplorer() const;	
+
 	/** Gets the image to display for the specified template */
 	const FSlateBrush* GetProjectItemImage( TWeakPtr<FProjectItem> ProjectItem ) const;
 
diff --git a/Engine/Source/Runtime/Core/Private/GenericPlatform/GenericPlatformMisc.cpp b/Engine/Source/Runtime/Core/Private/GenericPlatform/GenericPlatformMisc.cpp
--- a/Engine/Source/Runtime/Core/Private/GenericPlatform/GenericPlatformMisc.cpp
+++ b/Engine/Source/Runtime/Core/Private/GenericPlatform/GenericPlatformMisc.cpp
@@ -680,4 +680,9 @@ FString FGenericPlatformMisc::GetDefaultLocale()
 #else
 	return TEXT("en");
 #endif
-}
\ No newline at end of file
+}
+
+FText FGenericPlatformMisc::GetFileManagerName()
+{
+	return NSLOCTEXT("GenericPlatform", "FileManagerName", "File Manager");
+}
diff --git a/Engine/Source/Runtime/Core/Private/Mac/MacPlatformMisc.cpp b/Engine/Source/Runtime/Core/Private/Mac/MacPlatformMisc.cpp
--- a/Engine/Source/Runtime/Core/Private/Mac/MacPlatformMisc.cpp
+++ b/Engine/Source/Runtime/Core/Private/Mac/MacPlatformMisc.cpp
@@ -1107,6 +1107,11 @@ FString FMacPlatformMisc::GetDefaultLocale()
 	return FString::Printf(TEXT("%s_%s"), langCode, countryCode);
 }
 
+FText FMacPlatformMisc::GetFileManagerName()
+{
+	return NSLOCTEXT("MacPlatform", "FileManagerName", "Finder");
+}
+
 /** Global pointer to crash handler */
 void (* GCrashHandlerPointer)(const FGenericCrashContext & Context) = NULL;
 
diff --git a/Engine/Source/Runtime/Core/Private/Windows/WindowsPlatformMisc.cpp b/Engine/Source/Runtime/Core/Private/Windows/WindowsPlatformMisc.cpp
--- a/Engine/Source/Runtime/Core/Private/Windows/WindowsPlatformMisc.cpp
+++ b/Engine/Source/Runtime/Core/Private/Windows/WindowsPlatformMisc.cpp
@@ -1852,3 +1852,8 @@ const TCHAR* FWindowsPlatformMisc::GetDefaultPathSeparator()
 {
 	return TEXT( "\\" );
 }
+
+FText FWindowsPlatformMisc::GetFileManagerName()
+{
+	return NSLOCTEXT("WindowsPlatform", "FileManagerName", "Explorer");
+}
diff --git a/Engine/Source/Runtime/Core/Public/GenericPlatform/GenericPlatformMisc.h b/Engine/Source/Runtime/Core/Public/GenericPlatform/GenericPlatformMisc.h
--- a/Engine/Source/Runtime/Core/Public/GenericPlatform/GenericPlatformMisc.h
+++ b/Engine/Source/Runtime/Core/Public/GenericPlatform/GenericPlatformMisc.h
@@ -657,6 +657,9 @@ public:
 		return false;
 	}
 
+	/** @return Get the name of the platform specific file manager (eg, Explorer on Windows, Finder on OS X) */
+	static FText GetFileManagerName();
+
 	/**
 	 * Sample the displayed pixel color from anywhere on the screen using the OS
 	 *
diff --git a/Engine/Source/Runtime/Core/Public/Mac/MacPlatformMisc.h b/Engine/Source/Runtime/Core/Public/Mac/MacPlatformMisc.h
--- a/Engine/Source/Runtime/Core/Public/Mac/MacPlatformMisc.h
+++ b/Engine/Source/Runtime/Core/Public/Mac/MacPlatformMisc.h
@@ -119,6 +119,9 @@ struct CORE_API FMacPlatformMisc : public FGenericPlatformMisc
 	static void SetCrashHandler(void (* CrashHandler)(const FGenericCrashContext & Context));
 
 	static FString GetDefaultLocale();
+
+	/** @return Get the name of the platform specific file manager (Finder) */
+	static FText GetFileManagerName();
 };
 
 #ifdef __OBJC__
diff --git a/Engine/Source/Runtime/Core/Public/Windows/WindowsPlatformMisc.h b/Engine/Source/Runtime/Core/Public/Windows/WindowsPlatformMisc.h
--- a/Engine/Source/Runtime/Core/Public/Windows/WindowsPlatformMisc.h
+++ b/Engine/Source/Runtime/Core/Public/Windows/WindowsPlatformMisc.h
@@ -185,6 +185,9 @@ struct CORE_API FWindowsPlatformMisc : public FGenericPlatformMisc
 	* @return Windows path separator.
 	*/
 	static const TCHAR* GetDefaultPathSeparator();
+
+	/** @return Get the name of the platform specific file manager (Explorer) */
+	static FText GetFileManagerName();
 };
 
 typedef FWindowsPlatformMisc FPlatformMisc;
