Update association when upgrading to new engine versions. Only use the association to determine a project's "up to date" property; it serves as a coarse upgrade path between engine versions (so internal/source users won't be nagged) and it allows re-associating a project with different GitHub builds.

diff --git a/Engine/Source/Editor/GameProjectGeneration/Private/GameProjectUtils.cpp b/Engine/Source/Editor/GameProjectGeneration/Private/GameProjectUtils.cpp
--- a/Engine/Source/Editor/GameProjectGeneration/Private/GameProjectUtils.cpp
+++ b/Engine/Source/Editor/GameProjectGeneration/Private/GameProjectUtils.cpp
@@ -383,49 +383,30 @@ void GameProjectUtils::CheckForOutOfDateGameProjectFile()
 		{
 			if ( !ProjectStatus.bUpToDate )
 			{
-				// For the time being, if it's writable, just overwrite it with the latest info. If we need to do anything potentially destructive, we can change this logic for this in future builds.
-				bool bHaveUpdatedProject = false;
-				if (!FPlatformFileManager::Get().GetPlatformFile().IsReadOnly(*LoadedProjectFilePath))
+				const FText UpdateProjectText = LOCTEXT("UpdateProjectFilePrompt", "Update project to open in this version of the editor by default?");
+				const FText UpdateProjectConfirmText = LOCTEXT("UpdateProjectFileConfirm", "Update");
+				const FText UpdateProjectCancelText = LOCTEXT("UpdateProjectFileCancel", "Not Now");
+
+				FNotificationInfo Info(UpdateProjectText);
+				Info.bFireAndForget = false;
+				Info.bUseLargeFont = false;
+				Info.bUseThrobber = false;
+				Info.bUseSuccessFailIcons = false;
+				Info.FadeOutDuration = 3.f;
+				Info.ButtonDetails.Add(FNotificationButtonInfo(UpdateProjectConfirmText, FText(), FSimpleDelegate::CreateStatic(&GameProjectUtils::OnUpdateProjectConfirm)));
+				Info.ButtonDetails.Add(FNotificationButtonInfo(UpdateProjectCancelText, FText(), FSimpleDelegate::CreateStatic(&GameProjectUtils::OnUpdateProjectCancel)));
+
+				if (UpdateGameProjectNotification.IsValid())
 				{
-					FText FailureReason;
-					if (IProjectManager::Get().UpdateLoadedProjectFileToCurrent(NULL, FailureReason))
-					{
-						bHaveUpdatedProject = true;
-					}
-					else
-					{
-						UE_LOG(LogGameProjectGeneration, Error, TEXT("Failed to auto-update project file %s"), *LoadedProjectFilePath, *FailureReason.ToString());
-					}
+					UpdateGameProjectNotification.Pin()->ExpireAndFadeout();
+					UpdateGameProjectNotification.Reset();
 				}
 
-				// Otherwise, fall back to prompting
-				if (!bHaveUpdatedProject)
-				{
-					const FText UpdateProjectText = LOCTEXT("UpdateProjectFilePrompt", "The project file you have loaded is not up to date. Would you like to update it now?");
-					const FText UpdateProjectConfirmText = LOCTEXT("UpdateProjectFileConfirm", "Update");
-					const FText UpdateProjectCancelText = LOCTEXT("UpdateProjectFileCancel", "Not Now");
-
-					FNotificationInfo Info(UpdateProjectText);
-					Info.bFireAndForget = false;
-					Info.bUseLargeFont = false;
-					Info.bUseThrobber = false;
-					Info.bUseSuccessFailIcons = false;
-					Info.FadeOutDuration = 3.f;
-					Info.ButtonDetails.Add(FNotificationButtonInfo(UpdateProjectConfirmText, FText(), FSimpleDelegate::CreateStatic(&GameProjectUtils::OnUpdateProjectConfirm)));
-					Info.ButtonDetails.Add(FNotificationButtonInfo(UpdateProjectCancelText, FText(), FSimpleDelegate::CreateStatic(&GameProjectUtils::OnUpdateProjectCancel)));
-
-					if (UpdateGameProjectNotification.IsValid())
-					{
-						UpdateGameProjectNotification.Pin()->ExpireAndFadeout();
-						UpdateGameProjectNotification.Reset();
-					}
-
-					UpdateGameProjectNotification = FSlateNotificationManager::Get().AddNotification(Info);
+				UpdateGameProjectNotification = FSlateNotificationManager::Get().AddNotification(Info);
 
-					if (UpdateGameProjectNotification.IsValid())
-					{
-						UpdateGameProjectNotification.Pin()->SetCompletionState(SNotificationItem::CS_Pending);
-					}
+				if (UpdateGameProjectNotification.IsValid())
+				{
+					UpdateGameProjectNotification.Pin()->SetCompletionState(SNotificationItem::CS_Pending);
 				}
 			}
 		}
diff --git a/Engine/Source/Editor/UnrealEd/Private/UnrealEd.cpp b/Engine/Source/Editor/UnrealEd/Private/UnrealEd.cpp
--- a/Engine/Source/Editor/UnrealEd/Private/UnrealEd.cpp
+++ b/Engine/Source/Editor/UnrealEd/Private/UnrealEd.cpp
@@ -101,7 +101,7 @@ int32 EditorInit( IEngineLoop& EngineLoop )
 	const bool bDoAutomatedMapBuild = FParse::Param( FCommandLine::Get(), TEXT("AutomatedMapBuild") );
 
 	// Prompt to update the game project file to the current version, if necessary
-	if ( FRocketSupport::IsRocket() && FPaths::IsProjectFilePathSet() )
+	if ( FPaths::IsProjectFilePathSet() )
 	{
 		FGameProjectGenerationModule::Get().CheckForOutOfDateGameProjectFile();
 	}
diff --git a/Engine/Source/Runtime/Projects/Private/PluginManagerShared.cpp b/Engine/Source/Runtime/Projects/Private/PluginManagerShared.cpp
--- a/Engine/Source/Runtime/Projects/Private/PluginManagerShared.cpp
+++ b/Engine/Source/Runtime/Projects/Private/PluginManagerShared.cpp
@@ -300,18 +300,20 @@ bool FProjectOrPlugin::IsUpToDate( ) const
 {
 	const FProjectOrPluginInfo& ProjectOrPluginInfo = GetProjectOrPluginInfo();
 
-	if ( ProjectOrPluginInfo.FileVersion < VER_LATEST_PROJECT_FILE
-		|| ProjectOrPluginInfo.PackageFileUE4Version < GPackageFileUE4Version
-		|| ProjectOrPluginInfo.EngineVersion.ToString() != GEngineVersion.ToString()
-		|| ProjectOrPluginInfo.PackageFileLicenseeUE4Version < GPackageFileLicenseeUE4Version )
+	if (ProjectOrPluginInfo.FileVersion < VER_LATEST_PROJECT_FILE)
 	{
-		// This project file is out of date in at least one category
 		return false;
 	}
-	else
+
+	FString CurrentEngineIdentifier;
+	FPlatformMisc::GetEngineIdentifierFromRootDir(FPlatformMisc::RootDir(), CurrentEngineIdentifier);
+
+	if (ProjectOrPluginInfo.EngineAssociation != CurrentEngineIdentifier)
 	{
-		return true;
+		return false;
 	}
+
+	return true;
 }
 
 void FProjectOrPlugin::UpdateVersionToCurrent( )
@@ -322,6 +324,8 @@ void FProjectOrPlugin::UpdateVersionToCurrent( )
 	ProjectOrPluginInfo.EngineVersion = GEngineVersion;
 	ProjectOrPluginInfo.PackageFileUE4Version = GPackageFileUE4Version;
 	ProjectOrPluginInfo.PackageFileLicenseeUE4Version = GPackageFileLicenseeUE4Version;
+
+	FPlatformMisc::GetEngineIdentifierFromRootDir(FPlatformMisc::RootDir(), ProjectOrPluginInfo.EngineAssociation);
 }
 
 void FProjectOrPlugin::ReplaceModulesInProject(const TArray<FString>* StartupModuleNames)
