#ttp 331271 - EDITOR: New Project Dialog doesn't help inform users when they have an illegal character in their Windows user name
#proj UE4
#branch UE4
#summary Verified why @ was disallowed
#extra Also improved the robustness of the path validator and disallowed some extra things that had previously been missed.
#extra Made the way we report invalid characters consistent between FPaths, FName, and FPackageName; they now report all matched invalid characters rather than only the first one.

diff --git a/Engine/Source/Editor/GameProjectGeneration/Private/GameProjectUtils.cpp b/Engine/Source/Editor/GameProjectGeneration/Private/GameProjectUtils.cpp
--- a/Engine/Source/Editor/GameProjectGeneration/Private/GameProjectUtils.cpp
+++ b/Engine/Source/Editor/GameProjectGeneration/Private/GameProjectUtils.cpp
@@ -81,12 +81,8 @@ bool GameProjectUtils::IsValidProjectFileForCreation(const FString& ProjectFile,
 		return false;
 	}
 
-	FString IllegalPathCharacters;
-	if ( !ProjectPathContainsOnlyLegalCharacters(FPaths::GetPath(ProjectFile), IllegalPathCharacters) )
+	if ( !FPaths::ValidatePath(FPaths::GetPath(ProjectFile), &OutFailReason) )
 	{
-		FFormatNamedArguments Args;
-		Args.Add( TEXT("IllegalPathCharacters"), FText::FromString( IllegalPathCharacters ) );
-		OutFailReason = FText::Format( LOCTEXT( "ProjectNameContainsIllegalCharacters", "Project names may not contain the following characters: {IllegalPathCharacters}" ), Args );
 		return false;
 	}
 
@@ -173,12 +169,8 @@ bool GameProjectUtils::OpenProject(const FString& ProjectFile, FText& OutFailRea
 		return false;
 	}
 
-	FString IllegalPathCharacters;
-	if ( !ProjectPathContainsOnlyLegalCharacters(FPaths::GetPath(ProjectFile), IllegalPathCharacters) )
+	if ( !FPaths::ValidatePath(FPaths::GetPath(ProjectFile), &OutFailReason) )
 	{
-		FFormatNamedArguments Args;
-		Args.Add( TEXT("IllegalPathCharacters"), FText::FromString( IllegalPathCharacters ) );
-		OutFailReason = FText::Format( LOCTEXT( "ProjectPathContainsIllegalCharacters", "A projects path may not contain the following characters: {IllegalPathCharacters}" ), Args );
 		return false;
 	}
 
@@ -1057,49 +1049,6 @@ bool GameProjectUtils::NameContainsOnlyLegalCharacters(const FString& TestName,
 	return !bContainsIllegalCharacters;
 }
 
-bool GameProjectUtils::ProjectPathContainsOnlyLegalCharacters(const FString& ProjectFilePath, FString& OutIllegalCharacters)
-{
-	bool bContainsIllegalCharacters = false;
-
-	// set of characters not allowed in paths (Windows limitations)
-	FString AllIllegalCharacters = TEXT("<>\"|?*");
-
-	// Add lower-32 characters (illegal on windows)
-	for ( TCHAR IllegalChar = 0; IllegalChar < 32; IllegalChar++ )
-	{
-		AllIllegalCharacters += IllegalChar;
-	}
-
-	// @ characters are not legal. Used for revision/label specifiers in P4/SVN.
-	AllIllegalCharacters += TEXT("@");					
-	// # characters are not legal. Used for revision specifiers in P4/SVN.
-	AllIllegalCharacters += TEXT("#");					
-	// ^ characters are not legal. While the filesystem wont complain about this character, visual studio will.
-	AllIllegalCharacters += TEXT("^");					
-
-#if PLATFORM_MAC
-	// : characters are illegal on Mac OSX (but we need to allow them on Windows!)
-	AllIllegalCharacters += TEXT(":");
-#endif
-
-	for ( int32 CharIdx = 0; CharIdx < AllIllegalCharacters.Len() ; ++CharIdx )
-	{
-		const FString& Char = AllIllegalCharacters.Mid( CharIdx, 1 );
-
-		if ( ProjectFilePath.Contains( Char ) )
-		{
-			if ( !OutIllegalCharacters.Contains( Char ) )
-			{
-				OutIllegalCharacters += Char;
-			}
-
-			bContainsIllegalCharacters = true;
-		}
-	}
-
-	return !bContainsIllegalCharacters;
-}
-
 bool GameProjectUtils::ProjectFileExists(const FString& ProjectFile)
 {
 	return FPlatformFileManager::Get().GetPlatformFile().FileExists(*ProjectFile);
diff --git a/Engine/Source/Editor/GameProjectGeneration/Private/GameProjectUtils.h b/Engine/Source/Editor/GameProjectGeneration/Private/GameProjectUtils.h
--- a/Engine/Source/Editor/GameProjectGeneration/Private/GameProjectUtils.h
+++ b/Engine/Source/Editor/GameProjectGeneration/Private/GameProjectUtils.h
@@ -75,9 +75,6 @@ private:
 	/** Checks the name for illegal characters */
 	static bool NameContainsOnlyLegalCharacters(const FString& TestName, FString& OutIllegalCharacters);
 
-	/** Checks the project file path for illegal characters */
-	static bool ProjectPathContainsOnlyLegalCharacters(const FString& ProjectFilePath, FString& OutIllegalCharacters);
-
 	/** Returns true if the project file exists on disk */
 	static bool ProjectFileExists(const FString& ProjectFile);
 
diff --git a/Engine/Source/Runtime/Core/Private/Misc/Paths.cpp b/Engine/Source/Runtime/Core/Private/Misc/Paths.cpp
--- a/Engine/Source/Runtime/Core/Private/Misc/Paths.cpp
+++ b/Engine/Source/Runtime/Core/Private/Misc/Paths.cpp
@@ -784,6 +784,81 @@ FString FPaths::CreateTempFilename( const TCHAR* Path, const TCHAR* Prefix, cons
 	return UniqueFilename;
 }
 
+bool FPaths::ValidatePath( const FString& InPath, FText* OutReason )
+{
+	// Windows has the most restricted file system, and since we're cross platform, we have to respect the limitations of the lowest common denominator
+	// # isn't legal. Used for revision specifiers in P4/SVN, and also not allowed on Windows anyway
+	// @ isn't legal. Used for revision/label specifiers in P4/SVN
+	// ^ isn't legal. While the file-system won't complain about this character, Visual Studio will				
+	static const FString RestrictedChars = "/?:&\\*\"<>|%#@^";
+	static const FString RestrictedNames[] = {	"CON", "PRN", "AUX", "CLOCK$", "NUL", 
+												"COM1", "COM2", "COM3", "COM4", "COM5", "COM6", "COM7", "COM8", "COM9", 
+												"LPT1", "LPT2", "LPT3", "LPT4", "LPT5", "LPT6", "LPT7", "LPT8", "LPT9" };
+
+	FString Standardized = InPath;
+	NormalizeFilename(Standardized);
+	CollapseRelativeDirectories(Standardized);
+
+	// Walk each part of the path looking for name errors
+	for(int32 StartPos = 0, EndPos = Standardized.Find(TEXT("/")); ; 
+		StartPos = EndPos + 1, EndPos = Standardized.Find(TEXT("/"), ESearchCase::IgnoreCase, ESearchDir::FromStart, StartPos)
+		)
+	{
+		const bool bIsLastPart = EndPos == INDEX_NONE;
+		const FString PathPart = Standardized.Mid(StartPos, (bIsLastPart) ? MAX_int32 : EndPos - StartPos);
+
+		// If this is the first part of the path, it's possible for it to be a drive name and is allowed to contain a colon
+		if(StartPos == 0 && IsDrive(PathPart))
+		{
+			continue;
+		}
+
+		// Check for invalid characters
+		TCHAR CharString[] = { '\0', '\0' };
+		FString MatchedInvalidChars;
+		for(const TCHAR* InvalidCharacters = *RestrictedChars; *InvalidCharacters; ++InvalidCharacters)
+		{
+			CharString[0] = *InvalidCharacters;
+			if(PathPart.Contains(CharString))
+			{
+				MatchedInvalidChars += *InvalidCharacters;
+			}
+		}
+		if(MatchedInvalidChars.Len())
+		{
+			if(OutReason)
+			{
+				FFormatNamedArguments Args;
+				Args.Add(TEXT("IllegalPathCharacters"), FText::FromString(MatchedInvalidChars));
+				*OutReason = FText::Format(NSLOCTEXT("Core", "PathContainsInvalidCharacters", "Path may not contain the following characters: {IllegalPathCharacters}"), Args);
+			}
+			return false;
+		}
+
+		// Check for reserved names
+		for(const FString& RestrictedName : RestrictedNames)
+		{
+			if(PathPart.Equals(RestrictedName, ESearchCase::IgnoreCase))
+			{
+				if(OutReason)
+				{
+					FFormatNamedArguments Args;
+					Args.Add(TEXT("RestrictedName"), FText::FromString(RestrictedName));
+					*OutReason = FText::Format(NSLOCTEXT("Core", "PathContainsRestrictedName", "Path may not contain a restricted name: {RestrictedName}"), Args);
+				}
+				return false;
+			}
+		}
+
+		if(bIsLastPart)
+		{
+			break;
+		}
+	}
+
+	return true;
+}
+
 void FPaths::Split( const FString& InPath, FString& PathPart, FString& FilenamePart, FString& ExtensionPart )
 {
 	PathPart = GetPath(InPath);
diff --git a/Engine/Source/Runtime/Core/Private/UObject/UnrealNames.cpp b/Engine/Source/Runtime/Core/Private/UObject/UnrealNames.cpp
--- a/Engine/Source/Runtime/Core/Private/UObject/UnrealNames.cpp
+++ b/Engine/Source/Runtime/Core/Private/UObject/UnrealNames.cpp
@@ -835,19 +835,26 @@ bool FName::IsValidXName( FString InvalidChars/*=INVALID_NAME_CHARACTERS*/, FTex
 	FString Name = ToString();
 
 	// See if the name contains invalid characters.
-	FString Char;
+	TCHAR CharString[] = { '\0', '\0' };
+	FString MatchedInvalidChars;
 	for( int32 x = 0; x < InvalidChars.Len() ; ++x )
 	{
-		Char = InvalidChars.Mid( x, 1 );
+		CharString[0] = InvalidChars[x];
+		if( Name.Contains( CharString ) )
+		{
+			MatchedInvalidChars += CharString;
+		}
+	}
 
-		if( Name.Contains( Char ) )
+	if ( MatchedInvalidChars.Len() )
+	{
+		if ( Reason )
 		{
-			if ( Reason != NULL )
-			{
-				*Reason = FText::Format( NSLOCTEXT("Engine", "FNameContainsInvalidCharacter", "Name contains an invalid character : [{0}]"), FText::FromString( Char ) );
-			}
-			return false;
+			FFormatNamedArguments Args;
+			Args.Add(TEXT("IllegalNameCharacters"), FText::FromString(MatchedInvalidChars));
+			*Reason = FText::Format( NSLOCTEXT("Core", "NameContainsInvalidCharacters", "Name may not contain the following characters: {IllegalNameCharacters}"), Args );
 		}
+		return false;
 	}
 
 	return true;
diff --git a/Engine/Source/Runtime/Core/Public/Misc/Paths.h b/Engine/Source/Runtime/Core/Public/Misc/Paths.h
--- a/Engine/Source/Runtime/Core/Public/Misc/Paths.h
+++ b/Engine/Source/Runtime/Core/Public/Misc/Paths.h
@@ -387,6 +387,15 @@ public:
 	 */
 	static FString CreateTempFilename( const TCHAR* Path, const TCHAR* Prefix = TEXT(""), const TCHAR* Extension = TEXT(".tmp") );
 
+	/** 
+	 * Validates that the parts that make up the path contain no invalid characters as dictated by the operating system
+	 * Note that this is a different set of restrictions to those imposed by FPackageName
+	 *
+	 * @param InPath - path to validate
+	 * @param OutReason - optional parameter to fill with the failure reason
+	 */
+	static bool ValidatePath( const FString& InPath, FText* OutReason = nullptr );
+
 	/**
 	 * Parses a fully qualified or relative filename into its components (filename, path, extension).
 	 *
diff --git a/Engine/Source/Runtime/CoreUObject/Private/Misc/PackageName.cpp b/Engine/Source/Runtime/CoreUObject/Private/Misc/PackageName.cpp
--- a/Engine/Source/Runtime/CoreUObject/Private/Misc/PackageName.cpp
+++ b/Engine/Source/Runtime/CoreUObject/Private/Misc/PackageName.cpp
@@ -380,17 +380,24 @@ bool FPackageName::DoesPackageNameContainInvalidCharacters(const FString& InLong
 {
 	// See if the name contains invalid characters.
 	TCHAR CharString[] = { '\0', '\0' };
+	FString MatchedInvalidChars;
 	for (const TCHAR* InvalidCharacters = INVALID_LONGPACKAGE_CHARACTERS; *InvalidCharacters; ++InvalidCharacters)
 	{
 		CharString[0] = *InvalidCharacters;
 		if (InLongPackageName.Contains(CharString))
 		{
-			if (OutReason != NULL)
-			{
-				*OutReason = FText::Format( NSLOCTEXT("Core", "NameContainsInvalidCharacter", "Name contains an invalid character : [{0}]"), FText::FromString( CharString ) );
-			}
-			return true;
+			MatchedInvalidChars += *InvalidCharacters;
+		}
+	}
+	if (MatchedInvalidChars.Len())
+	{
+		if (OutReason)
+		{
+			FFormatNamedArguments Args;
+			Args.Add( TEXT("IllegalNameCharacters"), FText::FromString(MatchedInvalidChars) );
+			*OutReason = FText::Format( NSLOCTEXT("Core", "NameContainsInvalidCharacters", "Name may not contain the following characters: {IllegalNameCharacters}"), Args );
 		}
+		return true;
 	}
 	return false;
 }
