#ttp 322244 - LIVE: Feature Request: During "Add Code to Project", allow user to select path where files go
#proj UE4
#branch UE4
#summary You can now choose where to place a class added via the New Class Wizard
#extra This tries to be smart about your placement if you have Public and Private folders for your project.
- By default the header would go into Public, and the source file would go into Private.
- If you select the Public/Classes folder for the path, the source file will still go into Private.
- If you have a sub-path, eg) /Public/MyStuff/MyClass.h, this will be mirrored in the placement of the source file, eg) /Private/MyStuff/MyClass.cpp
#extra If you're not using Public or Private folders it will just place the source at whatever path you specified.
#extra It will verify that your source code is going to a valid module folder for your game, and also allows matching of modules that start with your game name, eg) MyGame, MyGameEditor.

diff --git a/Engine/Source/Editor/GameProjectGeneration/Private/GameProjectUtils.cpp b/Engine/Source/Editor/GameProjectGeneration/Private/GameProjectUtils.cpp
--- a/Engine/Source/Editor/GameProjectGeneration/Private/GameProjectUtils.cpp
+++ b/Engine/Source/Editor/GameProjectGeneration/Private/GameProjectUtils.cpp
@@ -527,9 +527,9 @@ bool GameProjectUtils::IsValidClassNameForCreation(const FString& NewClassName,
 	return true;
 }
 
-bool GameProjectUtils::AddCodeToProject(const FString& NewClassName, const UClass* ParentClass, FString& OutHeaderFilePath, FString& OutCppFilePath, FText& OutFailReason)
+bool GameProjectUtils::AddCodeToProject(const FString& NewClassName, const FString& NewClassPath, const UClass* ParentClass, FString& OutHeaderFilePath, FString& OutCppFilePath, FText& OutFailReason)
 {
-	const bool bAddCodeSuccessful = AddCodeToProject_Internal(NewClassName, ParentClass, OutHeaderFilePath, OutCppFilePath, OutFailReason);
+	const bool bAddCodeSuccessful = AddCodeToProject_Internal(NewClassName, NewClassPath, ParentClass, OutHeaderFilePath, OutCppFilePath, OutFailReason);
 
 	if( FEngineAnalytics::IsAvailable() )
 	{
@@ -1388,6 +1388,127 @@ bool GameProjectUtils::IsStarterContentAvailableForNewProjects()
 	return (StarterContentFiles.Num() > 0);
 }
 
+FString GameProjectUtils::GetSourceRootPath(const bool bIncludeModuleName)
+{
+	FString SourceDir = FPaths::GameSourceDir();
+
+	if(bIncludeModuleName)
+	{
+		// Assuming the game name is the same as the primary game module name
+		const FString ModuleName = FApp::GetGameName();
+		SourceDir /= ModuleName;
+	}
+
+	SourceDir /= "";
+
+	return FPaths::ConvertRelativePathToFull(SourceDir);
+}
+
+bool GameProjectUtils::IsValidSourcePath(const FString& InPath, const bool bIncludeModuleName, FText* const OutFailReason)
+{
+	FString RootPath = GetSourceRootPath(bIncludeModuleName);
+	if(bIncludeModuleName)
+	{
+		// If we're including the module name, then we want to allow variations of it, eg) MyModule, MyModuleEditor, MyModuleClient
+		// Those variations are valid, so we trim the last / from the path so that the StartsWith check below allows these variations
+		RootPath = RootPath.LeftChop(1); // Trim trailing /
+	}
+
+	const FString AbsoluteInPath = FPaths::ConvertRelativePathToFull(InPath) / ""; // Ensure trailing /
+
+	// Validate the path contains no invalid characters
+	if(!FPaths::ValidatePath(AbsoluteInPath, OutFailReason))
+	{
+		return false;
+	}
+
+	if(!AbsoluteInPath.StartsWith(RootPath))
+	{
+		if(OutFailReason)
+		{
+			FFormatNamedArguments Args;
+			Args.Add(TEXT("RootSourcePath"), FText::FromString(GetSourceRootPath(bIncludeModuleName)));
+			const FText FormatString = (bIncludeModuleName)
+				? LOCTEXT("SourcePathInvalidModuleRoot", "All source code must exist within a valid module folder in the projects source path, e.g. {RootSourcePath}")
+				: LOCTEXT("SourcePathInvalidRoot", "All source code must exist within the projects source path: {RootSourcePath}");
+			*OutFailReason = FText::Format( FormatString, Args );
+		}
+		return false;
+	}
+
+	return true;
+}
+
+bool GameProjectUtils::CalculateSourcePaths(const FString& InPath, FString& OutModuleName, FString& OutHeaderPath, FString& OutSourcePath, FText* const OutFailReason)
+{
+	const FString AbsoluteInPath = FPaths::ConvertRelativePathToFull(InPath) / ""; // Ensure trailing /
+	OutHeaderPath = AbsoluteInPath;
+	OutSourcePath = AbsoluteInPath;
+	OutModuleName.Empty();
+
+	if(!IsValidSourcePath(InPath, true/*bIncludeModuleName*/, OutFailReason))
+	{
+		return false;
+	}
+
+	// We've validated that this path includes a partial match for our module (eg, MyModule, MyModuleEditor, MyModuleClient)
+	// so extract the actual name of the module from the path so that we can generate the internal folder names correctly
+	const FString BaseRootPath = GetSourceRootPath(false/*bIncludeModuleName*/);
+	const int32 ModuleNameStartIndex = BaseRootPath.Len();
+	const int32 ModuleNameEndIndex = AbsoluteInPath.Find(TEXT("/"), ESearchCase::IgnoreCase, ESearchDir::FromStart, ModuleNameStartIndex);
+	check(ModuleNameEndIndex != INDEX_NONE); // this should never happen since AbsoluteInPath ends in a /, and we verified it started with BaseRootPath in IsValidSourcePath
+	OutModuleName = AbsoluteInPath.Mid(ModuleNameStartIndex, ModuleNameEndIndex - ModuleNameStartIndex);
+
+	const FString RootPath = BaseRootPath / OutModuleName / ""; // Ensure trailing /
+	const FString ClassesPath = RootPath / "Classes" / "";		// Ensure trailing /
+	const FString PublicPath = RootPath / "Public" / "";		// Ensure trailing /
+	const FString PrivatePath = RootPath / "Private" / "";		// Ensure trailing /
+
+	// The root path must exist; we will allow the creation of sub-folders, but not the module root!
+	if(!IFileManager::Get().DirectoryExists(*RootPath))
+	{
+		if(OutFailReason)
+		{
+			FFormatNamedArguments Args;
+			Args.Add(TEXT("ModuleSourcePath"), FText::FromString(RootPath));
+			*OutFailReason = FText::Format(LOCTEXT("SourcePathMissingModuleRoot", "The specified module path does not exist on disk: {ModuleSourcePath}"), Args);
+		}
+		return false;
+	}
+
+	// If either the Public or Private path exists, and we're in the root, force the header/source file to use one of these folders
+	const bool bPublicPathExists = IFileManager::Get().DirectoryExists(*PublicPath);
+	const bool bPrivatePathExists = IFileManager::Get().DirectoryExists(*PrivatePath);
+	const bool bForceInternalPath = AbsoluteInPath == RootPath && (bPublicPathExists || bPrivatePathExists);
+
+	// The rules for placing header files are as follows:
+	// 1) If InPath is the source root, check to see if there is a Public folder within it, and if so, place the header there
+	// 2) Otherwise, just place the header at InPath (the default set above)
+	if(AbsoluteInPath == RootPath)
+	{
+		OutHeaderPath = (bPublicPathExists || bForceInternalPath) ? PublicPath : AbsoluteInPath;
+	}
+
+	// The rules for placing source files are as follows:
+	// 1) If InPath is the source root, check to see if there is a Private folder within it, and if so, place the source file there
+	// 2) If InPath is contained within the Public or Classes folder of this module, place it in the equivalent path in the Private folder
+	// 3) Otherwise, just place the source file at InPath (the default set above)
+	if(AbsoluteInPath == RootPath)
+	{
+		OutSourcePath = (bPrivatePathExists || bForceInternalPath) ? PrivatePath : AbsoluteInPath;
+	}
+	else if(AbsoluteInPath.StartsWith(ClassesPath))
+	{
+		OutSourcePath = AbsoluteInPath.Replace(*ClassesPath, *PrivatePath);
+	}
+	else if(AbsoluteInPath.StartsWith(PublicPath))
+	{
+		OutSourcePath = AbsoluteInPath.Replace(*PublicPath, *PrivatePath);
+	}
+
+	return !OutHeaderPath.IsEmpty() && !OutSourcePath.IsEmpty();
+}
+
 bool GameProjectUtils::ReadTemplateFile(const FString& TemplateFileName, FString& OutFileContents, FText& OutFailReason)
 {
 	const FString FullFileName = FPaths::EngineContentDir() / TEXT("Editor") / TEXT("Templates") / TemplateFileName;
@@ -1872,7 +1993,7 @@ bool GameProjectUtils::ProjectHasCodeFiles()
 	return GameProjectUtils::GetProjectCodeFileCount() > 0;
 }
 
-bool GameProjectUtils::AddCodeToProject_Internal(const FString& NewClassName, const UClass* ParentClass, FString& OutHeaderFilePath, FString& OutCppFilePath, FText& OutFailReason)
+bool GameProjectUtils::AddCodeToProject_Internal(const FString& NewClassName, const FString& NewClassPath, const UClass* ParentClass, FString& OutHeaderFilePath, FString& OutCppFilePath, FText& OutFailReason)
 {
 	if ( !ParentClass )
 	{
@@ -1891,19 +2012,27 @@ bool GameProjectUtils::AddCodeToProject_Internal(const FString& NewClassName, co
 		return false;
 	}
 
+	FString ModuleName;
+	FString NewHeaderPath;
+	FString NewCppPath;
+	if ( !CalculateSourcePaths(NewClassPath, ModuleName, NewHeaderPath, NewCppPath, &OutFailReason) )
+	{
+		return false;
+	}
+
 	const bool bAllowNewSlowTask = true;
 	FStatusMessageContext SlowTaskMessage( LOCTEXT( "AddingCodeToProject", "Adding code to project..." ), bAllowNewSlowTask );
 
-	const FString SourceDir = FPaths::GameSourceDir().LeftChop(1); // Trim the trailing /
-
-	// Assuming the game name is the same as the primary game module name
-	const FString ModuleName = FApp::GetGameName();
-	const FString ModulePath = SourceDir / ModuleName;
-
 	// If the project does not already contain code, add the primary game module
 	TArray<FString> CreatedFiles;
 	if ( !ProjectHasCodeFiles() )
 	{
+		// We always add the basic source code to the root directory, not the potential sub-directory provided by NewClassPath
+		const FString SourceDir = FPaths::GameSourceDir().LeftChop(1); // Trim the trailing /
+
+		// Assuming the game name is the same as the primary game module name
+		const FString ModuleName = FApp::GetGameName();
+
 		TArray<FString> StartupModuleNames;
 		if ( GenerateBasicSourceCode(SourceDir, ModuleName, StartupModuleNames, CreatedFiles, OutFailReason) )
 		{
@@ -1918,7 +2047,7 @@ bool GameProjectUtils::AddCodeToProject_Internal(const FString& NewClassName, co
 
 	// Class Header File
 	FString SyncLocation;
-	const FString NewHeaderFilename = ModulePath / NewClassName + TEXT(".h");
+	const FString NewHeaderFilename = NewHeaderPath / NewClassName + TEXT(".h");
 	{
 		if ( GenerateClassHeaderFile(NewHeaderFilename, ParentClass, TArray<FString>(), TEXT(""), TEXT(""), SyncLocation, OutFailReason) )
 		{
@@ -1926,22 +2055,22 @@ bool GameProjectUtils::AddCodeToProject_Internal(const FString& NewClassName, co
 		}
 		else
 		{
-			DeleteCreatedFiles(SourceDir, CreatedFiles);
+			DeleteCreatedFiles(NewHeaderPath, CreatedFiles);
 			return false;
 		}
 	}
 
 	// Class CPP file
-	const FString NewCPPFilename = ModulePath / NewClassName + TEXT(".cpp");
+	const FString NewCppFilename = NewCppPath / NewClassName + TEXT(".cpp");
 	{
 		const FString PrefixedClassName = FString(ParentClass->GetPrefixCPP()) + NewClassName;
-		if ( GenerateClassCPPFile(NewCPPFilename, ModuleName, PrefixedClassName, TArray<FString>(), TArray<FString>(), TEXT(""), OutFailReason) )
+		if ( GenerateClassCPPFile(NewCppFilename, ModuleName, PrefixedClassName, TArray<FString>(), TArray<FString>(), TEXT(""), OutFailReason) )
 		{
-			CreatedFiles.Add(NewCPPFilename);
+			CreatedFiles.Add(NewCppFilename);
 		}
 		else
 		{
-			DeleteCreatedFiles(SourceDir, CreatedFiles);
+			DeleteCreatedFiles(NewCppPath, CreatedFiles);
 			return false;
 		}
 	}
@@ -1967,7 +2096,7 @@ bool GameProjectUtils::AddCodeToProject_Internal(const FString& NewClassName, co
 	}
 
 	OutHeaderFilePath = NewHeaderFilename;
-	OutCppFilePath = NewCPPFilename;
+	OutCppFilePath = NewCppFilename;
 
 	return true;
 }
diff --git a/Engine/Source/Editor/GameProjectGeneration/Private/GameProjectUtils.h b/Engine/Source/Editor/GameProjectGeneration/Private/GameProjectUtils.h
--- a/Engine/Source/Editor/GameProjectGeneration/Private/GameProjectUtils.h
+++ b/Engine/Source/Editor/GameProjectGeneration/Private/GameProjectUtils.h
@@ -33,7 +33,7 @@ public:
 	static bool IsValidClassNameForCreation(const FString& NewClassName, FText& OutFailReason);
 
 	/** Adds new source code to the project. When returning true, OutSyncFileAndLineNumber will be the the preferred target file to sync in the users code editing IDE, formatted for use with GenericApplication::GotoLineInSource */
-	static bool AddCodeToProject(const FString& NewClassName, const UClass* ParentClass, FString& OutHeaderFilePath, FString& OutCppFilePath, FText& OutFailReason);
+	static bool AddCodeToProject(const FString& NewClassName, const FString& NewClassPath, const UClass* ParentClass, FString& OutHeaderFilePath, FString& OutCppFilePath, FText& OutFailReason);
 
 	/** Loads a template project definitions object from the TemplateDefs.ini file in the specified project */
 	static UTemplateProjectDefs* LoadTemplateDefs(const FString& ProjectDirectory);
@@ -56,6 +56,39 @@ public:
 	/** Returns true if there are starter content files available for instancing into new projects. */
 	static bool IsStarterContentAvailableForNewProjects();
 
+	/** 
+	 * Get the absolute root path under which all project source code must exist
+	 *
+	 * @param	bIncludeModuleName	Whether to include the module name in the root path?
+	 * 
+	 * @return	The root path. Will always be an absolute path ending with a /
+	 */
+	static FString GetSourceRootPath(const bool bIncludeModuleName);
+
+	/** 
+	 * Check to see if the given path is a valid place to put source code for this project (exists within the source root path) 
+	 *
+	 * @param	InPath				The path to check
+	 * @param	bIncludeModuleName	Whether to require the module name in the root path? (is really used as a prefix so will allow MyModule, MyModuleEditor, etc)
+	 * @param	OutFailReason		Optional parameter to fill with failure information
+	 * 
+	 * @return	true if the path is valid, false otherwise
+	 */
+	static bool IsValidSourcePath(const FString& InPath, const bool bIncludeModuleName, FText* const OutFailReason = nullptr);
+
+	/** 
+	 * Given the path provided, work out where generated .h and .cpp files would be placed
+	 *
+	 * @param	InPath				The path to use a base
+	 * @param	OutModuleName		The module name extracted from the path (the part after GetSourceRootPath(false))
+	 * @param	OutHeaderPath		The path where the .h file should be placed
+	 * @param	OutSourcePath		The path where the .cpp file should be placed
+	 * @param	OutFailReason		Optional parameter to fill with failure information
+	 * 
+	 * @return	false if the paths are invalid
+	 */
+	static bool CalculateSourcePaths(const FString& InPath, FString& OutModuleName, FString& OutHeaderPath, FString& OutSourcePath, FText* const OutFailReason = nullptr);
+
 private:
 	/** Generates a new project without using a template project */
 	static bool GenerateProjectFromScratch(const FString& NewProjectFile, bool bShouldGenerateCode, bool bCopyStarterContent, FText& OutFailReason);
@@ -183,7 +216,7 @@ private:
 	static bool ProjectHasCodeFiles();
 
 	/** Internal handler for AddCodeToProject*/
-	static bool AddCodeToProject_Internal(const FString& NewClassName, const UClass* ParentClass, FString& OutHeaderFilePath, FString& OutCppFilePath, FText& OutFailReason);
+	static bool AddCodeToProject_Internal(const FString& NewClassName, const FString& NewClassPath, const UClass* ParentClass, FString& OutHeaderFilePath, FString& OutCppFilePath, FText& OutFailReason);
 
 private:
 	static TWeakPtr<SNotificationItem> UpdateGameProjectNotification;
diff --git a/Engine/Source/Editor/GameProjectGeneration/Private/SNewClassDialog.cpp b/Engine/Source/Editor/GameProjectGeneration/Private/SNewClassDialog.cpp
--- a/Engine/Source/Editor/GameProjectGeneration/Private/SNewClassDialog.cpp
+++ b/Engine/Source/Editor/GameProjectGeneration/Private/SNewClassDialog.cpp
@@ -7,6 +7,7 @@
 #include "ClassViewerModule.h"
 #include "ClassViewerFilter.h"
 #include "Editor/ClassViewer/Private/SClassViewer.h"
+#include "DesktopPlatformModule.h"
 
 #define LOCTEXT_NAMESPACE "GameProjectGeneration"
 
@@ -56,18 +57,20 @@ public:
 BEGIN_SLATE_FUNCTION_BUILD_OPTIMIZATION
 void SNewClassDialog::Construct( const FArguments& InArgs )
 {
+	NewClassPath = GameProjectUtils::GetSourceRootPath(true/*bIncludeModuleName*/);
+
 	ParentClass = InArgs._Class;
 
 	DialogFixedWidth = 900;
 	bShowFullClassTree = false;
 
-	LastNameValidityCheckTime = 0;
-	NameValidityCheckFrequency = 4;
-	bLastNameValidityCheckSuccessful = true;
+	LastPeriodicValidityCheckTime = 0;
+	PeriodicValidityCheckFrequency = 4;
+	bLastInputValidityCheckSuccessful = true;
 	bPreventPeriodicValidityChecksUntilNextChange = false;
 
 	SetupParentClassItems();
-	UpdateClassNameValidity();
+	UpdateInputValidity();
 
 	FClassViewerInitializationOptions Options;
 	Options.Mode = EClassViewerMode::ClassPicker;
@@ -84,6 +87,8 @@ void SNewClassDialog::Construct( const FArguments& InArgs )
 
 	ClassViewer = StaticCastSharedRef<SClassViewer>(FModuleManager::LoadModuleChecked<FClassViewerModule>("ClassViewer").CreateClassViewer(Options, FOnClassPicked::CreateSP(this, &SNewClassDialog::OnAdvancedClassSelected)));
 
+	const float EditableTextHeight = 26.0f;
+
 	ChildSlot
 	[
 		SNew(SBorder)
@@ -296,11 +301,11 @@ void SNewClassDialog::Construct( const FArguments& InArgs )
 							.AutoHeight()
 							.Padding(0)
 							[
-								SNew(SHorizontalBox)
+								SNew(SGridPanel)
+								.FillColumn(1, 1.0f)
 
 								// Name label
-								+SHorizontalBox::Slot()
-								.AutoWidth()
+								+SGridPanel::Slot(0, 0)
 								.VAlign(VAlign_Center)
 								.Padding(0, 0, 12, 0)
 								[
@@ -310,13 +315,105 @@ void SNewClassDialog::Construct( const FArguments& InArgs )
 								]
 
 								// Name edit box
-								+SHorizontalBox::Slot()
-								.FillWidth(1.f)
+								+SGridPanel::Slot(1, 0)
+								.Padding(0.0f, 3.0f)
 								.VAlign(VAlign_Center)
 								[
-									SAssignNew( ClassNameEditBox, SEditableTextBox)
-									.Text( this, &SNewClassDialog::OnGetClassNameText )
-									.OnTextChanged( this, &SNewClassDialog::OnClassNameTextChanged )
+									SNew(SBox)
+									.HeightOverride(EditableTextHeight)
+									[
+										SAssignNew( ClassNameEditBox, SEditableTextBox)
+										.Text( this, &SNewClassDialog::OnGetClassNameText )
+										.OnTextChanged( this, &SNewClassDialog::OnClassNameTextChanged )
+									]
+								]
+
+								// Path label
+								+SGridPanel::Slot(0, 1)
+								.VAlign(VAlign_Center)
+								.Padding(0, 0, 12, 0)
+								[
+									SNew(STextBlock)
+									.TextStyle( FEditorStyle::Get(), "NewClassDialog.SelectedParentClassLabel" )
+									.Text( LOCTEXT( "PathLabel", "Path" ).ToString() )
+								]
+
+								// Path edit box
+								+SGridPanel::Slot(1, 1)
+								.Padding(0.0f, 3.0f)
+								.VAlign(VAlign_Center)
+								[
+									SNew(SBox)
+									.HeightOverride(EditableTextHeight)
+									[
+										SNew(SHorizontalBox)
+
+										+SHorizontalBox::Slot()
+										.FillWidth(1.0f)
+										[
+											SNew(SEditableTextBox)
+											.Text(this, &SNewClassDialog::OnGetClassPathText)
+											.OnTextChanged(this, &SNewClassDialog::OnClassPathTextChanged)
+										]
+
+										+SHorizontalBox::Slot()
+										.AutoWidth()
+										.Padding(6.0f, 1.0f, 0.0f, 0.0f)
+										[
+											SNew(SButton)
+											.VAlign(VAlign_Center)
+											.OnClicked(this, &SNewClassDialog::HandleChooseFolderButtonClicked)
+											.Text( LOCTEXT( "BrowseButtonText", "Choose Folder" ) )
+										]
+									]
+								]
+
+								// Header output label
+								+SGridPanel::Slot(0, 2)
+								.VAlign(VAlign_Center)
+								.Padding(0, 0, 12, 0)
+								[
+									SNew(STextBlock)
+									.TextStyle( FEditorStyle::Get(), "NewClassDialog.SelectedParentClassLabel" )
+									.Text( LOCTEXT( "HeaderFileLabel", "Header File" ).ToString() )
+								]
+
+								// Header output text
+								+SGridPanel::Slot(1, 2)
+								.Padding(0.0f, 3.0f)
+								.VAlign(VAlign_Center)
+								[
+									SNew(SBox)
+									.VAlign(VAlign_Center)
+									.HeightOverride(EditableTextHeight)
+									[
+										SNew(STextBlock)
+										.Text(this, &SNewClassDialog::OnGetClassHeaderFileText)
+									]
+								]
+
+								// Source output label
+								+SGridPanel::Slot(0, 3)
+								.VAlign(VAlign_Center)
+								.Padding(0, 0, 12, 0)
+								[
+									SNew(STextBlock)
+									.TextStyle( FEditorStyle::Get(), "NewClassDialog.SelectedParentClassLabel" )
+									.Text( LOCTEXT( "SourceFileLabel", "Source File" ).ToString() )
+								]
+
+								// Source output text
+								+SGridPanel::Slot(1, 3)
+								.Padding(0.0f, 3.0f)
+								.VAlign(VAlign_Center)
+								[
+									SNew(SBox)
+									.VAlign(VAlign_Center)
+									.HeightOverride(EditableTextHeight)
+									[
+										SNew(STextBlock)
+										.Text(this, &SNewClassDialog::OnGetClassSourceFileText)
+									]
 								]
 							]
 						]
@@ -369,12 +466,12 @@ void SNewClassDialog::Tick( const FGeometry& AllottedGeometry, const double InCu
 {
 	SCompoundWidget::Tick(AllottedGeometry, InCurrentTime, InDeltaTime);
 
-	// Every few seconds, the class name is checked for validity in case the disk contents changed and the location is now valid or invalid.
+	// Every few seconds, the class name/path is checked for validity in case the disk contents changed and the location is now valid or invalid.
 	// After class creation, periodic checks are disabled to prevent a brief message indicating that the class you created already exists.
 	// This feature is re-enabled if the user did not restart and began editing parameters again.
-	if ( !bPreventPeriodicValidityChecksUntilNextChange && (InCurrentTime > LastNameValidityCheckTime + NameValidityCheckFrequency) )
+	if ( !bPreventPeriodicValidityChecksUntilNextChange && (InCurrentTime > LastPeriodicValidityCheckTime + PeriodicValidityCheckFrequency) )
 	{
-		UpdateClassNameValidity();
+		UpdateInputValidity();
 	}
 }
 
@@ -518,9 +615,9 @@ EVisibility SNewClassDialog::GetNameErrorLabelVisibility() const
 
 FString SNewClassDialog::GetNameErrorLabelText() const
 {
-	if ( !bLastNameValidityCheckSuccessful )
+	if ( !bLastInputValidityCheckSuccessful )
 	{
-		return LastNameValidityErrorText.ToString();
+		return LastInputValidityErrorText.ToString();
 	}
 
 	return TEXT("");
@@ -548,9 +645,20 @@ FString SNewClassDialog::GetGlobalErrorLabelText() const
 
 void SNewClassDialog::OnNamePageEntered()
 {
+	// Set the default class name based on the selected parent class, eg MyActor
+	FString PotentialNewClassName = "My";
+	PotentialNewClassName += GetSelectedParentClassName();
+
+	// Only set the default if the user hasn't changed the class name from the previous default
+	if(LastAutoGeneratedClassName.IsEmpty() || NewClassName == LastAutoGeneratedClassName)
+	{
+		NewClassName = PotentialNewClassName;
+		LastAutoGeneratedClassName = PotentialNewClassName;
+		UpdateInputValidity();
+	}
+
 	// Steal keyboard focus to accelerate name entering
 	FSlateApplication::Get().SetKeyboardFocus(ClassNameEditBox, EKeyboardFocusCause::SetDirectly);
-
 }
 
 FString SNewClassDialog::GetNameClassTitle() const
@@ -566,7 +674,28 @@ FText SNewClassDialog::OnGetClassNameText() const
 void SNewClassDialog::OnClassNameTextChanged(const FText& NewText)
 {
 	NewClassName = NewText.ToString();
-	UpdateClassNameValidity();
+	UpdateInputValidity();
+}
+
+FText SNewClassDialog::OnGetClassPathText() const
+{
+	return FText::FromString(NewClassPath);
+}
+
+void SNewClassDialog::OnClassPathTextChanged(const FText& NewText)
+{
+	NewClassPath = NewText.ToString();
+	UpdateInputValidity();
+}
+
+FText SNewClassDialog::OnGetClassHeaderFileText() const
+{
+	return FText::FromString(CalculatedClassHeaderName);
+}
+
+FText SNewClassDialog::OnGetClassSourceFileText() const
+{
+	return FText::FromString(CalculatedClassSourceName);
 }
 
 void SNewClassDialog::CancelClicked()
@@ -576,7 +705,7 @@ void SNewClassDialog::CancelClicked()
 
 bool SNewClassDialog::CanFinish() const
 {
-	return bLastNameValidityCheckSuccessful && GetSelectedParentClass() != NULL && FSourceCodeNavigation::IsCompilerAvailable();
+	return bLastInputValidityCheckSuccessful && GetSelectedParentClass() != NULL && FSourceCodeNavigation::IsCompilerAvailable();
 }
 
 void SNewClassDialog::FinishClicked()
@@ -587,7 +716,7 @@ void SNewClassDialog::FinishClicked()
 	FString CppFilePath;
 
 	FText FailReason;
-	if ( GameProjectUtils::AddCodeToProject(NewClassName, GetSelectedParentClass(), HeaderFilePath, CppFilePath, FailReason) )
+	if ( GameProjectUtils::AddCodeToProject(NewClassName, NewClassPath, GetSelectedParentClass(), HeaderFilePath, CppFilePath, FailReason) )
 	{
 		// Prevent periodic validity checks. This is to prevent a brief error message about the class already existing while you are exiting.
 		bPreventPeriodicValidityChecksUntilNextChange = true;
@@ -629,20 +758,55 @@ void SNewClassDialog::OnDownloadIDEClicked(FString URL)
 	FPlatformProcess::LaunchURL( *URL, NULL, NULL );
 }
 
-void SNewClassDialog::UpdateClassNameValidity()
+FReply SNewClassDialog::HandleChooseFolderButtonClicked()
 {
-	if ( NewClassName.IsEmpty() )
+	IDesktopPlatform* DesktopPlatform = FDesktopPlatformModule::Get();
+	if ( DesktopPlatform )
 	{
-		// An empty class name is not valid, but don't show error text since it implies the user did something wrong.
-		LastNameValidityErrorText = FText();
-		bLastNameValidityCheckSuccessful = false;
+		TSharedPtr<SWindow> ParentWindow = FSlateApplication::Get().FindWidgetWindow(AsShared());
+		void* ParentWindowWindowHandle = (ParentWindow.IsValid()) ? ParentWindow->GetNativeWindow()->GetOSWindowHandle() : nullptr;
+
+		FString FolderName;
+		const FString Title = LOCTEXT("NewClassBrowseTitle", "Choose a source location").ToString();
+		const bool bFolderSelected = DesktopPlatform->OpenDirectoryDialog(
+			ParentWindowWindowHandle,
+			Title,
+			NewClassPath,
+			FolderName
+			);
+
+		if ( bFolderSelected )
+		{
+			if ( !FolderName.EndsWith(TEXT("/")) )
+			{
+				FolderName += TEXT("/");
+			}
+
+			NewClassPath = FolderName;
+			UpdateInputValidity();
+		}
 	}
-	else
+
+	return FReply::Handled();
+}
+
+void SNewClassDialog::UpdateInputValidity()
+{
+	bLastInputValidityCheckSuccessful = true;
+
+	// Validate the path first since this has the side effect of updating the UI
+	FString ModuleName;
+	bLastInputValidityCheckSuccessful = GameProjectUtils::CalculateSourcePaths(NewClassPath, ModuleName, CalculatedClassHeaderName, CalculatedClassSourceName, &LastInputValidityErrorText);
+	CalculatedClassHeaderName /= NewClassName + ".h";
+	CalculatedClassSourceName /= NewClassName + ".cpp";
+
+	// Validate the class name only if the path is valid
+	if ( bLastInputValidityCheckSuccessful )
 	{
-		bLastNameValidityCheckSuccessful = GameProjectUtils::IsValidClassNameForCreation(NewClassName, LastNameValidityErrorText);
+		bLastInputValidityCheckSuccessful = GameProjectUtils::IsValidClassNameForCreation(NewClassName, LastInputValidityErrorText);
 	}
 
-	LastNameValidityCheckTime = FSlateApplication::Get().GetCurrentTime();
+	LastPeriodicValidityCheckTime = FSlateApplication::Get().GetCurrentTime();
 
 	// Since this function was invoked, periodic validity checks should be re-enabled if they were disabled.
 	bPreventPeriodicValidityChecksUntilNextChange = false;
diff --git a/Engine/Source/Editor/GameProjectGeneration/Private/SNewClassDialog.h b/Engine/Source/Editor/GameProjectGeneration/Private/SNewClassDialog.h
--- a/Engine/Source/Editor/GameProjectGeneration/Private/SNewClassDialog.h
+++ b/Engine/Source/Editor/GameProjectGeneration/Private/SNewClassDialog.h
@@ -80,6 +80,18 @@ private:
 	/** Handler for when the text in the class name edit box has changed */
 	void OnClassNameTextChanged(const FText& NewText);
 
+	/** Returns the text in the class path edit box */
+	FText OnGetClassPathText() const;
+
+	/** Handler for when the text in the class path edit box has changed */
+	void OnClassPathTextChanged(const FText& NewText);
+
+	/** Returns the text for the calculated header file name */
+	FText OnGetClassHeaderFileText() const;
+
+	/** Returns the text for the calculated source file name */
+	FText OnGetClassSourceFileText() const;
+
 	/** Handler for when cancel is clicked */
 	void CancelClicked();
 
@@ -92,9 +104,12 @@ private:
 	/** Handler for when the error label IDE hyperlink was clicked */
 	void OnDownloadIDEClicked(FString URL);
 
+	/** Handler for when the "Choose Folder" button is clicked */
+	FReply HandleChooseFolderButtonClicked( );
+
 private:
-	/** Checks the current class name for validity and updates cached values accordingly */
-	void UpdateClassNameValidity();
+	/** Checks the current class name/path for validity and updates cached values accordingly */
+	void UpdateInputValidity();
 
 	/** Gets the currently selected parent class */
 	const UClass* GetSelectedParentClass() const;
@@ -123,6 +138,18 @@ private:
 	/** The name of the class being created */
 	FString NewClassName;
 
+	/** The path to place the files for the class being generated */
+	FString NewClassPath;
+
+	/** The calculated name of the generated header file for this class */
+	FString CalculatedClassHeaderName;
+
+	/** The calculated name of the generated source file for this class */
+	FString CalculatedClassSourceName;
+
+	/** The name of the last class that was auto-generated by this wizard */
+	FString LastAutoGeneratedClassName;
+
 	/** The selected parent class */
 	TWeakObjectPtr<UClass> ParentClass;
 
@@ -132,18 +159,18 @@ private:
 	/** If true, the full class tree will be shown in the parent class selection */
 	bool bShowFullClassTree;
 
-	/** The last time that the class name was checked for validity. This is used to throttle I/O requests to a reasonable frequency */
-	double LastNameValidityCheckTime;
+	/** The last time that the class name/path was checked for validity. This is used to throttle I/O requests to a reasonable frequency */
+	double LastPeriodicValidityCheckTime;
 
-	/** The frequency in seconds for validity checks while the dialog is idle. Changes to the path immediately update the validity. */
-	double NameValidityCheckFrequency;
+	/** The frequency in seconds for validity checks while the dialog is idle. Changes to the name/path immediately update the validity. */
+	double PeriodicValidityCheckFrequency;
 
 	/** Periodic checks for validity will not occur while this flag is true. Used to prevent a frame of "this project already exists" while exiting after a successful creation. */
 	bool bPreventPeriodicValidityChecksUntilNextChange;
 
 	/** The error text from the last validity check */
-	FText LastNameValidityErrorText;
+	FText LastInputValidityErrorText;
 
-	/** True if the last validity check returned that the class name is valid for creation */
-	bool bLastNameValidityCheckSuccessful;
+	/** True if the last validity check returned that the class name/path is valid for creation */
+	bool bLastInputValidityCheckSuccessful;
 };
