#ttp 328688 - Project Browser doesn't refresh for new downloads
#summary Added ability to refresh the project browser (using the refresh button, or by pressing F5)
#proj Editor
#branch UE4

diff --git a/Engine/Source/Editor/EditorStyle/Private/SlateEditorStyle.cpp b/Engine/Source/Editor/EditorStyle/Private/SlateEditorStyle.cpp
--- a/Engine/Source/Editor/EditorStyle/Private/SlateEditorStyle.cpp
+++ b/Engine/Source/Editor/EditorStyle/Private/SlateEditorStyle.cpp
@@ -247,6 +247,7 @@ void FSlateEditorStyle::FStyle::SetupGeneralStyles()
 		Set( "Icons.Info", new IMAGE_BRUSH( "Icons/icon_info_16x", Icon16x16) );
 		Set( "Icons.Warning", new IMAGE_BRUSH( "Icons/icon_warning_16x", Icon16x16) );
 		Set( "Icons.Download", new IMAGE_BRUSH( "Icons/icon_Downloads_16x", Icon16x16) );
+		Set( "Icons.Refresh", new IMAGE_BRUSH( "Icons/icon_Refresh_16x", Icon16x16 ) );
 	}
 
 	Set( "WarningStripe", new IMAGE_BRUSH( "Common/WarningStripe", FVector2D(20,6), FLinearColor::White, ESlateBrushTileType::Horizontal ) );
diff --git a/Engine/Source/Editor/GameProjectGeneration/Private/SProjectBrowser.cpp b/Engine/Source/Editor/GameProjectGeneration/Private/SProjectBrowser.cpp
--- a/Engine/Source/Editor/GameProjectGeneration/Private/SProjectBrowser.cpp
+++ b/Engine/Source/Editor/GameProjectGeneration/Private/SProjectBrowser.cpp
@@ -65,18 +65,13 @@ void SProjectBrowser::Construct( const FArguments& InArgs )
 	bPreventSelectionChangeEvent = false;
 	ThumbnailBorderPadding = 5;
 	ThumbnailSize = 128.0f;
+	bAllowProjectCreate = InArgs._AllowProjectCreate;
 
 	// Prepare the categories box
 	CategoriesBox = SNew(SVerticalBox);
 
 	// Find all projects
-	FindProjects(InArgs._AllowProjectCreate);
-
-	for (auto CategoryIt = ProjectCategories.CreateConstIterator(); CategoryIt; ++CategoryIt)
-	{
-		const TSharedRef<FProjectCategory>& Category = *CategoryIt;
-		ConstructCategory( CategoriesBox.ToSharedRef(), Category );
-	}
+	FindProjects();
 
 	CategoriesBox->AddSlot()
 	.HAlign(HAlign_Center)
@@ -124,20 +119,36 @@ void SProjectBrowser::Construct( const FArguments& InArgs )
 				.VAlign(VAlign_Top)
 				.Padding(FMargin(5.f, 10.f))
 				[
-					SNew(SOverlay)
-
-					+SOverlay::Slot()
+					SNew(SHorizontalBox)
+					+SHorizontalBox::Slot()
 					[
-						SNew(SSearchBox)
-						.HintText(LOCTEXT("FilterHint", "Filter Projects..."))
-						.OnTextChanged(this, &SProjectBrowser::OnFilterTextChanged)
+						SNew(SOverlay)
+
+						+SOverlay::Slot()
+						[
+							SNew(SSearchBox)
+							.HintText(LOCTEXT("FilterHint", "Filter Projects..."))
+							.OnTextChanged(this, &SProjectBrowser::OnFilterTextChanged)
+						]
+
+						+SOverlay::Slot()
+						[
+							SNew(SBorder)
+							.Visibility(this, &SProjectBrowser::GetFilterActiveOverlayVisibility)
+							.BorderImage(FEditorStyle::Get().GetBrush("SearchBox.ActiveBorder"))
+						]
 					]
 
-					+SOverlay::Slot()
+					+SHorizontalBox::Slot()
+					.AutoWidth()
 					[
-						SNew(SBorder)
-						.Visibility(this, &SProjectBrowser::GetFilterActiveOverlayVisibility)
-						.BorderImage(FEditorStyle::Get().GetBrush("SearchBox.ActiveBorder"))
+						SNew(SButton)
+						.ButtonStyle(FCoreStyle::Get(), "NoBorder")
+						.OnClicked(this, &SProjectBrowser::FindProjects)
+						.ToolTipText(LOCTEXT("RefreshProjectList", "Refresh the project list"))
+						[
+							SNew(SImage).Image(FEditorStyle::Get().GetBrush("Icons.Refresh"))
+						]
 					]
 				]
 
@@ -522,8 +533,11 @@ FText SProjectBrowser::GetSelectedProjectName() const
 	return FText::GetEmpty();
 }
 
-void SProjectBrowser::FindProjects(bool bAllowProjectCreate)
+FReply SProjectBrowser::FindProjects()
 {
+	ProjectCategories.Empty();
+	CategoriesBox->ClearChildren();
+
 	const FText CreateProjectCategoryName = LOCTEXT("NewProjectCategoryName", "Create Project");
 	const FText MyProjectsCategoryName = LOCTEXT("MyProjectsCategoryName", "My Projects");
 	const FText ProjectsCategoryName = LOCTEXT("ProjectsCategoryName", "My Projects");
@@ -763,6 +777,14 @@ void SProjectBrowser::FindProjects(bool bAllowProjectCreate)
 	}
 
 	PopulateFilteredProjectCategories();
+
+	for (auto CategoryIt = ProjectCategories.CreateConstIterator(); CategoryIt; ++CategoryIt)
+	{
+		const TSharedRef<FProjectCategory>& Category = *CategoryIt;
+		ConstructCategory(CategoriesBox.ToSharedRef(), Category);
+	}
+
+	return FReply::Handled();
 }
 
 
@@ -808,6 +830,16 @@ void SProjectBrowser::PopulateFilteredProjectCategories()
 	}
 }
 
+FReply SProjectBrowser::OnKeyDown( const FGeometry& MyGeometry, const FKeyboardEvent& InKeyboardEvent )
+{
+	if (InKeyboardEvent.GetKey() == EKeys::F5)
+	{
+		return FindProjects();
+	}
+
+	return FReply::Unhandled();
+}
+
 
 bool SProjectBrowser::OpenProject( const FString& ProjectFile )
 {
diff --git a/Engine/Source/Editor/GameProjectGeneration/Private/SProjectBrowser.h b/Engine/Source/Editor/GameProjectGeneration/Private/SProjectBrowser.h
--- a/Engine/Source/Editor/GameProjectGeneration/Private/SProjectBrowser.h
+++ b/Engine/Source/Editor/GameProjectGeneration/Private/SProjectBrowser.h
@@ -69,7 +69,7 @@ protected:
 	FText GetSelectedProjectName( ) const;
 
 	/** Populates ProjectItemsSource with projects found on disk */
-	void FindProjects( bool bAllowProjectCreate );
+	FReply FindProjects( );
 
 	/** Adds the specified project to the specified category. Creates a new category if necessary. */
 	void AddProjectToCategory( const TSharedRef<FProjectItem>& ProjectItem, const FText& ProjectCategory );
@@ -82,6 +82,16 @@ protected:
 
 	/** Populate the list of filtered project categories */
 	void PopulateFilteredProjectCategories();
+	
+	/**
+	 * Called after a key is pressed when this widget has keyboard focus (this event bubbles if not handled)
+	 *
+	 * @param MyGeometry The Geometry of the widget receiving the event
+	 * @param  InKeyboardEvent  Keyboard event
+	 *
+	 * @return  Returns whether the event was handled, along with other possible actions
+	 */
+	virtual FReply OnKeyDown( const FGeometry& MyGeometry, const FKeyboardEvent& InKeyboardEvent ) OVERRIDE;
 
 private:
 
@@ -141,6 +151,8 @@ private:
 
 	bool bHasProjectFiles;
 
+	/** True if this UI will enable the creation of new projects, false otherwise */
+	bool bAllowProjectCreate;
 private:
 
 	// Holds a delegate that is executed when the new project screen is being requested.
