#ttp 311016 CHECKIN: ROCKET: BUG: New Project wizard displays malformed path when hand-typing folder names; also shows forward slashes
#proj Editor
#branch UE4
#summary Enforce the platform specific path handling for "/" or "\\".

diff --git a/Engine/Source/Editor/GameProjectGeneration/Private/SNewProjectWizard.cpp b/Engine/Source/Editor/GameProjectGeneration/Private/SNewProjectWizard.cpp
--- a/Engine/Source/Editor/GameProjectGeneration/Private/SNewProjectWizard.cpp
+++ b/Engine/Source/Editor/GameProjectGeneration/Private/SNewProjectWizard.cpp
@@ -710,6 +710,7 @@ FString SNewProjectWizard::GetCurrentProjectFileParentFolder() const
 void SNewProjectWizard::OnCurrentProjectFilePathChanged(const FText& InValue)
 {
 	CurrentProjectFilePath = InValue.ToString();
+	FPaths::MakePlatformFilename(CurrentProjectFilePath);
 	UpdateProjectFileValidity();
 }
 
@@ -732,8 +733,8 @@ FString SNewProjectWizard::GetProjectFilenameWithPath() const
 		const FString ProjectName = CurrentProjectFileName;
 		const FString ProjectPath = IFileManager::Get().ConvertToAbsolutePathForExternalAppForWrite(*CurrentProjectFilePath);
 		const FString Filename = ProjectName + TEXT(".") + IProjectManager::GetProjectFileExtension();
-		const FString ProjectFilename = FPaths::Combine( *ProjectPath, *ProjectName, *Filename );
-
+		FString ProjectFilename = FPaths::Combine( *ProjectPath, *ProjectName, *Filename );
+		FPaths::MakePlatformFilename(ProjectFilename);
 		return ProjectFilename;
 	}
 }
@@ -768,7 +769,8 @@ FReply SNewProjectWizard::HandleBrowseButtonClicked()
 			{
 				FolderName += TEXT("/");
 			}
-
+			
+			FPaths::MakePlatformFilename(FolderName);
 			LastBrowsePath = FolderName;
 			CurrentProjectFilePath = FolderName;
 		}
@@ -1005,6 +1007,7 @@ void SNewProjectWizard::SetDefaultProjectLocation( )
 	}
 
 	FPaths::NormalizeFilename(DefaultProjectFilePath);
+	FPaths::MakePlatformFilename(DefaultProjectFilePath);
 	const FString GenericProjectName = LOCTEXT("DefaultProjectName", "MyProject").ToString();
 	FString ProjectName = GenericProjectName;
 
@@ -1036,6 +1039,7 @@ void SNewProjectWizard::SetDefaultProjectLocation( )
 	{
 		CurrentProjectFileName = ProjectName;
 		CurrentProjectFilePath = DefaultProjectFilePath;
+		FPaths::MakePlatformFilename(CurrentProjectFilePath);
 		LastBrowsePath = CurrentProjectFilePath;
 	}
 }
