#ttp 332842 - CRITICAL: Regression: EDITOR: No source folder when creating a blank project
#proj UE4
#branch UE4
#summary Made the source path validation more lenient when adding code to a project for the first time
#change GameProjectUtils::CalculateSourcePaths no longer checks that the module source folder exists on disk if the project contains no code.
#change GameProjectUtils::IsValidSourcePath no longer allows partial module name matches (to match both MyModule, and MyModuleEditor) if the project contains no code.

diff --git a/Engine/Source/Editor/GameProjectGeneration/Private/GameProjectUtils.cpp b/Engine/Source/Editor/GameProjectGeneration/Private/GameProjectUtils.cpp
--- a/Engine/Source/Editor/GameProjectGeneration/Private/GameProjectUtils.cpp
+++ b/Engine/Source/Editor/GameProjectGeneration/Private/GameProjectUtils.cpp
@@ -1407,7 +1407,10 @@ FString GameProjectUtils::GetSourceRootPath(const bool bIncludeModuleName)
 bool GameProjectUtils::IsValidSourcePath(const FString& InPath, const bool bIncludeModuleName, FText* const OutFailReason)
 {
 	FString RootPath = GetSourceRootPath(bIncludeModuleName);
-	if(bIncludeModuleName)
+
+	// Only allow partial module name matches if we already have code; the first class added to a project *must* be for the game module
+	const bool bHasCodeFiles = GameProjectUtils::ProjectHasCodeFiles();
+	if(bIncludeModuleName && bHasCodeFiles)
 	{
 		// If we're including the module name, then we want to allow variations of it, eg) MyModule, MyModuleEditor, MyModuleClient
 		// Those variations are valid, so we trim the last / from the path so that the StartsWith check below allows these variations
@@ -1465,7 +1468,9 @@ bool GameProjectUtils::CalculateSourcePaths(const FString& InPath, FString& OutM
 	const FString PrivatePath = RootPath / "Private" / "";		// Ensure trailing /
 
 	// The root path must exist; we will allow the creation of sub-folders, but not the module root!
-	if(!IFileManager::Get().DirectoryExists(*RootPath))
+	// We ignore this check if the project doesn't already have source code in it, as the module folder won't yet have been created
+	const bool bHasCodeFiles = GameProjectUtils::ProjectHasCodeFiles();
+	if(!IFileManager::Get().DirectoryExists(*RootPath) && bHasCodeFiles)
 	{
 		if(OutFailReason)
 		{
