#ttp 332687 - CRITICAL: Regression: Editor hangs after selecting File > New/Open Project
#summary Fixed performance regression caused by enumerating all engine installations for every project in the project browser
#proj Editor
#branch UE4

diff --git a/Engine/Source/Editor/GameProjectGeneration/Private/SProjectBrowser.cpp b/Engine/Source/Editor/GameProjectGeneration/Private/SProjectBrowser.cpp
--- a/Engine/Source/Editor/GameProjectGeneration/Private/SProjectBrowser.cpp
+++ b/Engine/Source/Editor/GameProjectGeneration/Private/SProjectBrowser.cpp
@@ -616,6 +616,8 @@ FReply SProjectBrowser::FindProjects()
 		}
 	}
 
+	const FString EngineIdentifier = FDesktopPlatformModule::Get()->GetCurrentEngineIdentifier();
+
 	// Add all discovered projects to the list
 	for ( auto ProjectFilenameIt = DiscoveredProjectFilesToCategory.CreateConstIterator(); ProjectFilenameIt; ++ProjectFilenameIt )
 	{
@@ -625,7 +627,7 @@ FReply SProjectBrowser::FindProjects()
 		{
 			const bool bPromptIfSavedWithNewerVersionOfEngine = false;
 			FProjectStatus ProjectStatus;
-			if (IProjectManager::Get().QueryStatusForProject(ProjectFilename, FDesktopPlatformModule::Get()->GetCurrentEngineIdentifier(), ProjectStatus))
+			if (IProjectManager::Get().QueryStatusForProject(ProjectFilename, EngineIdentifier, ProjectStatus))
 			{
 				// @todo localized project name
 				const FText ProjectName = FText::FromString(ProjectStatus.Name);
