Fill in the binary editor's file association when freshening a sample project during a Rocket build, rather than using the current engine's file association string (build machines won't have one).

diff --git a/Engine/Source/Editor/GameProjectGeneration/Private/GameProjectGenerationModule.cpp b/Engine/Source/Editor/GameProjectGeneration/Private/GameProjectGenerationModule.cpp
--- a/Engine/Source/Editor/GameProjectGeneration/Private/GameProjectGenerationModule.cpp
+++ b/Engine/Source/Editor/GameProjectGeneration/Private/GameProjectGenerationModule.cpp
@@ -50,9 +50,9 @@ void FGameProjectGenerationModule::CheckForOutOfDateGameProjectFile()
 }
 
 
-bool FGameProjectGenerationModule::UpdateGameProject()
+bool FGameProjectGenerationModule::UpdateGameProject(const FString& EngineIdentifier)
 {
-	return GameProjectUtils::UpdateGameProject();
+	return GameProjectUtils::UpdateGameProject(EngineIdentifier);
 }
 
 
diff --git a/Engine/Source/Editor/GameProjectGeneration/Private/GameProjectUtils.cpp b/Engine/Source/Editor/GameProjectGeneration/Private/GameProjectUtils.cpp
--- a/Engine/Source/Editor/GameProjectGeneration/Private/GameProjectUtils.cpp
+++ b/Engine/Source/Editor/GameProjectGeneration/Private/GameProjectUtils.cpp
@@ -407,13 +407,13 @@ void GameProjectUtils::CheckForOutOfDateGameProjectFile()
 	}
 }
 
-bool GameProjectUtils::UpdateGameProject()
+bool GameProjectUtils::UpdateGameProject(const FString& EngineIdentifier)
 {
 	const FString& ProjectFilename = FPaths::IsProjectFilePathSet() ? FPaths::GetProjectFilePath() : FString();
 	if ( !ProjectFilename.IsEmpty() )
 	{
 		FProjectStatus ProjectStatus;
-		if ( IProjectManager::Get().QueryStatusForProject(ProjectFilename, FDesktopPlatformModule::Get()->GetCurrentEngineIdentifier(), ProjectStatus) )
+		if ( IProjectManager::Get().QueryStatusForProject(ProjectFilename, EngineIdentifier, ProjectStatus) )
 		{
 			if ( ProjectStatus.bUpToDate )
 			{
@@ -424,7 +424,7 @@ bool GameProjectUtils::UpdateGameProject()
 			{
 				FText FailReason;
 				bool bWasCheckedOut = false;
-				if ( !UpdateGameProjectFile(ProjectFilename, NULL, bWasCheckedOut, FailReason) )
+				if ( !UpdateGameProjectFile(ProjectFilename, EngineIdentifier, NULL, bWasCheckedOut, FailReason) )
 				{
 					// The user chose to update, but the update failed. Notify the user.
 					UE_LOG(LogGameProjectGeneration, Error, TEXT("%s failed to update. %s"), *ProjectFilename, *FailReason.ToString() );
@@ -1855,7 +1855,7 @@ void GameProjectUtils::UpdateProject(const TArray<FString>* StartupModuleNames)
 	FText UpdateMessage;
 	SNotificationItem::ECompletionState NewCompletionState;
 	bool bWasCheckedOut = false;
-	if ( UpdateGameProjectFile(ProjectFilename, StartupModuleNames, bWasCheckedOut, FailReason) )
+	if ( UpdateGameProjectFile(ProjectFilename, FDesktopPlatformModule::Get()->GetCurrentEngineIdentifier(), StartupModuleNames, bWasCheckedOut, FailReason) )
 	{
 		// The project was updated successfully.
 		FFormatNamedArguments Args;
@@ -1896,7 +1896,7 @@ void GameProjectUtils::OnUpdateProjectCancel()
 	}
 }
 
-bool GameProjectUtils::UpdateGameProjectFile(const FString& ProjectFilename, const TArray<FString>* StartupModuleNames, bool& OutbWasCheckedOut, FText& OutFailReason)
+bool GameProjectUtils::UpdateGameProjectFile(const FString& ProjectFilename, const FString& EngineIdentifier, const TArray<FString>* StartupModuleNames, bool& OutbWasCheckedOut, FText& OutFailReason)
 {
 	// First attempt to check out the file if SCC is enabled
 	if ( ISourceControlModule::Get().IsEnabled() )
@@ -1924,7 +1924,7 @@ bool GameProjectUtils::UpdateGameProjectFile(const FString& ProjectFilename, con
 	}
 
 	// Now tell the project manager to update the file
-	if (!IProjectManager::Get().UpdateLoadedProjectFileToCurrent(StartupModuleNames, FDesktopPlatformModule::Get()->GetCurrentEngineIdentifier(), OutFailReason))
+	if (!IProjectManager::Get().UpdateLoadedProjectFileToCurrent(StartupModuleNames, EngineIdentifier, OutFailReason))
 	{
 		return false;
 	}
diff --git a/Engine/Source/Editor/GameProjectGeneration/Private/GameProjectUtils.h b/Engine/Source/Editor/GameProjectGeneration/Private/GameProjectUtils.h
--- a/Engine/Source/Editor/GameProjectGeneration/Private/GameProjectUtils.h
+++ b/Engine/Source/Editor/GameProjectGeneration/Private/GameProjectUtils.h
@@ -24,7 +24,7 @@ public:
 	static void CheckForOutOfDateGameProjectFile();
 
 	/** Updates the currently loaded project. Returns true if the project was updated successfully or if no update was needed */
-	static bool UpdateGameProject();
+	static bool UpdateGameProject(const FString &EngineIdentifier);
 
 	/** Opens a dialog to add code files to a project */
 	static void OpenAddCodeToProjectDialog();
@@ -207,7 +207,7 @@ private:
 
 	 * @return true, if successful
 	 */
-	static bool UpdateGameProjectFile(const FString& ProjectFilename, const TArray<FString>* StartupModuleNames, bool& OutbWasCheckedOut, FText& OutFailReason);
+	static bool UpdateGameProjectFile(const FString& ProjectFilename, const FString& EngineIdentifier, const TArray<FString>* StartupModuleNames, bool& OutbWasCheckedOut, FText& OutFailReason);
 
 	/** Checks the specified game project file out from source control */
 	static bool CheckoutGameProjectFile(const FString& ProjectFilename, FText& OutFailReason);
diff --git a/Engine/Source/Editor/GameProjectGeneration/Public/GameProjectGenerationModule.h b/Engine/Source/Editor/GameProjectGeneration/Public/GameProjectGenerationModule.h
--- a/Engine/Source/Editor/GameProjectGeneration/Public/GameProjectGenerationModule.h
+++ b/Engine/Source/Editor/GameProjectGeneration/Public/GameProjectGenerationModule.h
@@ -51,7 +51,7 @@ public:
 	virtual void CheckForOutOfDateGameProjectFile();
 
 	/** Updates the currently loaded project. Returns true if the project was updated successfully or if no update was needed */
-	virtual bool UpdateGameProject();
+	virtual bool UpdateGameProject(const FString &EngineIdentifier);
 
 	/** Updates the current code project */
 	virtual bool UpdateCodeProject(FText& OutFailReason);
diff --git a/Engine/Source/Editor/UnrealEd/Private/Commandlets/UpdateGameProjectCommandlet.cpp b/Engine/Source/Editor/UnrealEd/Private/Commandlets/UpdateGameProjectCommandlet.cpp
--- a/Engine/Source/Editor/UnrealEd/Private/Commandlets/UpdateGameProjectCommandlet.cpp
+++ b/Engine/Source/Editor/UnrealEd/Private/Commandlets/UpdateGameProjectCommandlet.cpp
@@ -67,9 +67,11 @@ int32 UUpdateGameProjectCommandlet::Main( const FString& InParams )
 		SourceControlProvider.Init();
 	}
 
-	UE_LOG(LogUpdateGameProjectCommandlet, Display, TEXT("Updating project file %s..."), *ProjectFilePath);
+	FString EngineIdentifier = GEngineVersion.ToString(EVersionComponent::Minor);
 
-	if ( !FGameProjectGenerationModule::Get().UpdateGameProject() )
+	UE_LOG(LogUpdateGameProjectCommandlet, Display, TEXT("Updating project file %s to %s..."), *ProjectFilePath, *EngineIdentifier);
+
+	if ( !FGameProjectGenerationModule::Get().UpdateGameProject(EngineIdentifier) )
 	{
 		// UpdateGameProject produces it's own error to the log
 		return 1;
