Remove all the different categories for sample games - just use "Samples" or "My Projects".

diff --git a/Engine/Source/Editor/GameProjectGeneration/Private/SProjectBrowser.cpp b/Engine/Source/Editor/GameProjectGeneration/Private/SProjectBrowser.cpp
--- a/Engine/Source/Editor/GameProjectGeneration/Private/SProjectBrowser.cpp
+++ b/Engine/Source/Editor/GameProjectGeneration/Private/SProjectBrowser.cpp
@@ -538,20 +538,8 @@ FReply SProjectBrowser::FindProjects()
 	ProjectCategories.Empty();
 	CategoriesBox->ClearChildren();
 
-	const FText CreateProjectCategoryName = LOCTEXT("NewProjectCategoryName", "Create Project");
 	const FText MyProjectsCategoryName = LOCTEXT("MyProjectsCategoryName", "My Projects");
-	const FText ProjectsCategoryName = LOCTEXT("ProjectsCategoryName", "My Projects");
-	const FText SampleGamesCategoryName = LOCTEXT("SampleGamesCategoryName", "Sample Games");
-	const FText DemoletsCategoryName = LOCTEXT("DemoletsCategoryName", "Demolets");
-	const FText ShowcasesCategoryName = LOCTEXT("ShowcasesCategoryName", "Showcases");
-
-	//// Add the option to create a new game to the "Create Project" category.
-	//if (bAllowProjectCreate)
-	//{
-	//	const bool bIsNewProjectItem = true;
-	//	TSharedRef<FProjectItem> NewProjectItem = MakeShareable(new FProjectItem(LOCTEXT("CreateNewProjectItem", "New Project"), FText(), NULL, TEXT(""), bIsNewProjectItem));
-	//	AddProjectToCategory(NewProjectItem, CreateProjectCategoryName);
-	//}
+	const FText SamplesCategoryName = LOCTEXT("SamplesCategoryName", "Samples");
 
 	// Form a list of all known project files.
 	TMap<FString, FText> DiscoveredProjectFilesToCategory;
@@ -587,11 +575,11 @@ FReply SProjectBrowser::FindProjects()
 
 	// @todo discover projects in the "holding tank" perhaps
 	// Hard-coding a few paths for now
-	ProjectFileDiscoverFoldersToCategory.Add(FFolderToCategory(FPaths::RootDir(), ProjectsCategoryName));
-	ProjectFileDiscoverFoldersToCategory.Add(FFolderToCategory(FPaths::Combine(*FPaths::RootDir(), TEXT("Samples")), SampleGamesCategoryName));
-	ProjectFileDiscoverFoldersToCategory.Add(FFolderToCategory(FPaths::Combine(*FPaths::RootDir(), TEXT("Samples"), TEXT("SampleGames")), SampleGamesCategoryName));
-	ProjectFileDiscoverFoldersToCategory.Add(FFolderToCategory(FPaths::Combine(*FPaths::RootDir(), TEXT("Samples"), TEXT("Showcases")), ShowcasesCategoryName));
-	ProjectFileDiscoverFoldersToCategory.Add(FFolderToCategory(FPaths::Combine(*FPaths::RootDir(), TEXT("Samples"), TEXT("Demolets")), DemoletsCategoryName));
+	ProjectFileDiscoverFoldersToCategory.Add(FFolderToCategory(FPaths::RootDir(), MyProjectsCategoryName));
+	ProjectFileDiscoverFoldersToCategory.Add(FFolderToCategory(FPaths::Combine(*FPaths::RootDir(), TEXT("Samples")), SamplesCategoryName));
+	ProjectFileDiscoverFoldersToCategory.Add(FFolderToCategory(FPaths::Combine(*FPaths::RootDir(), TEXT("Samples"), TEXT("SampleGames")), SamplesCategoryName));
+	ProjectFileDiscoverFoldersToCategory.Add(FFolderToCategory(FPaths::Combine(*FPaths::RootDir(), TEXT("Samples"), TEXT("Showcases")), SamplesCategoryName));
+	ProjectFileDiscoverFoldersToCategory.Add(FFolderToCategory(FPaths::Combine(*FPaths::RootDir(), TEXT("Samples"), TEXT("Demolets")), SamplesCategoryName));
 
 	// Discover all unique project files in all discover folders
 	for (auto RootFolderIt = ProjectFileDiscoverFoldersToCategory.CreateConstIterator(); RootFolderIt; ++RootFolderIt)
@@ -655,27 +643,11 @@ FReply SProjectBrowser::FindProjects()
 				FText ProjectCategory;
 				if ( ProjectStatus.bSignedSampleProject )
 				{
-					if(ProjectStatus.Category == "Sample Games")
-					{
-						ProjectCategory = SampleGamesCategoryName;
-					}
-					else if(ProjectStatus.Category == "Demolets")
-					{
-						ProjectCategory = DemoletsCategoryName;
-					}
-					else if(ProjectStatus.Category == "Showcases")
-					{
-						ProjectCategory = ShowcasesCategoryName;
-					}
+					ProjectCategory = SamplesCategoryName;
 				}
 				else
 				{
-					ProjectCategory = (ProjectStatus.Category.IsEmpty()) ? DetectedCategory : FText::FromString(ProjectStatus.Category);
-				}
-
-				if ( ProjectCategory.IsEmpty() )
-				{
-					ProjectCategory = MyProjectsCategoryName;
+					ProjectCategory = DetectedCategory;
 				}
 
 				const bool bIsNewProjectItem = false;
@@ -685,47 +657,23 @@ FReply SProjectBrowser::FindProjects()
 		}
 	}
 
-	// Make sure the category order is "Create Project", "My Projects", "Projects", "Sample Games", "Demolets", "Showcases" then all remaining categories in alphabetical order
-	TSharedPtr<FProjectCategory> CreateProjectCategory;
+	// Make sure the category order is "My Projects", "Samples", then all remaining categories in alphabetical order
 	TSharedPtr<FProjectCategory> MyProjectsCategory;
-	TSharedPtr<FProjectCategory> ProjectsCategory;
 	TSharedPtr<FProjectCategory> SamplesCategory;
-	TSharedPtr<FProjectCategory> DemoletsCategory;
-	TSharedPtr<FProjectCategory> ShowcasesCategory;
 
 	for ( int32 CategoryIdx = ProjectCategories.Num() - 1; CategoryIdx >= 0; --CategoryIdx )
 	{
 		TSharedRef<FProjectCategory> Category = ProjectCategories[CategoryIdx];
-		if ( Category->CategoryName.EqualTo(CreateProjectCategoryName) )
-		{
-			CreateProjectCategory = Category;
-			ProjectCategories.RemoveAt(CategoryIdx);
-		}
-		else if ( Category->CategoryName.EqualTo(MyProjectsCategoryName) )
+		if ( Category->CategoryName.EqualTo(MyProjectsCategoryName) )
 		{
 			MyProjectsCategory = Category;
 			ProjectCategories.RemoveAt(CategoryIdx);
 		}
-		else if ( Category->CategoryName.EqualTo(ProjectsCategoryName) )
-		{
-			ProjectsCategory = Category;
-			ProjectCategories.RemoveAt(CategoryIdx);
-		}
-		else if ( Category->CategoryName.EqualTo(SampleGamesCategoryName) )
+		else if ( Category->CategoryName.EqualTo(SamplesCategoryName) )
 		{
 			SamplesCategory = Category;
 			ProjectCategories.RemoveAt(CategoryIdx);
 		}
-		else if ( Category->CategoryName.EqualTo(DemoletsCategoryName) )
-		{
-			DemoletsCategory = Category;
-			ProjectCategories.RemoveAt(CategoryIdx);
-		}
-		else if ( Category->CategoryName.EqualTo(ShowcasesCategoryName) )
-		{
-			ShowcasesCategory = Category;
-			ProjectCategories.RemoveAt(CategoryIdx);
-		}
 	}
 
 	// Sort categories
@@ -738,31 +686,15 @@ FReply SProjectBrowser::FindProjects()
 	};
 	ProjectCategories.Sort( FCompareCategories() );
 
-	// Now readd the built-in categories (last added is first in the list)
-	if ( DemoletsCategory.IsValid() )
-	{
-		ProjectCategories.Insert(DemoletsCategory.ToSharedRef(), 0);
-	}
-	if ( ShowcasesCategory.IsValid() )
-	{
-		ProjectCategories.Insert(ShowcasesCategory.ToSharedRef(), 0);
-	}
+	// Now read the built-in categories (last added is first in the list)
 	if ( SamplesCategory.IsValid() )
 	{
 		ProjectCategories.Insert(SamplesCategory.ToSharedRef(), 0);
 	}
-	if ( ProjectsCategory.IsValid() )
-	{
-		ProjectCategories.Insert(ProjectsCategory.ToSharedRef(), 0);
-	}
 	if ( MyProjectsCategory.IsValid() )
 	{
 		ProjectCategories.Insert(MyProjectsCategory.ToSharedRef(), 0);
 	}
-	if ( CreateProjectCategory.IsValid() )
-	{
-		ProjectCategories.Insert(CreateProjectCategory.ToSharedRef(), 0);
-	}
 
 	// Sort each individual category
 	struct FCompareProjectItems
