Change the name length check to launch a descriptive warning after opening a project rather than error out on the request to open.
#ttp 333947 the landscape_worldmachine project cannot be opened on my machine because the project name is longer than 20 characters.

diff --git a/Engine/Source/Editor/GameProjectGeneration/Private/GameProjectGenerationModule.cpp b/Engine/Source/Editor/GameProjectGeneration/Private/GameProjectGenerationModule.cpp
--- a/Engine/Source/Editor/GameProjectGeneration/Private/GameProjectGenerationModule.cpp
+++ b/Engine/Source/Editor/GameProjectGeneration/Private/GameProjectGenerationModule.cpp
@@ -77,4 +77,9 @@ bool FGameProjectGenerationModule::UpdateCodeResourceFiles(TArray<FString>& OutC
 	return GameProjectUtils::GenerateGameResourceFiles(GameModuleSourcePath, FApp::GetGameName(), OutCreatedFiles, OutFailReason);
 }
 
+void FGameProjectGenerationModule::CheckAndWarnProjectFilenameValid()
+{
+	GameProjectUtils::CheckAndWarnProjectFilenameValid();
+}
+
 #undef LOCTEXT_NAMESPACE
diff --git a/Engine/Source/Editor/GameProjectGeneration/Private/GameProjectUtils.cpp b/Engine/Source/Editor/GameProjectGeneration/Private/GameProjectUtils.cpp
--- a/Engine/Source/Editor/GameProjectGeneration/Private/GameProjectUtils.cpp
+++ b/Engine/Source/Editor/GameProjectGeneration/Private/GameProjectUtils.cpp
@@ -21,6 +21,7 @@ checkAtCompileTime(PLATFORM_MAX_FILEPATH_LENGTH - MAX_PROJECT_PATH_BUFFER_SPACE
 #define MAX_CLASS_NAME_LENGTH 32 // Enforce a reasonable class name length so the path is not too long for PLATFORM_MAX_FILEPATH_LENGTH
 
 TWeakPtr<SNotificationItem> GameProjectUtils::UpdateGameProjectNotification = NULL;
+TWeakPtr<SNotificationItem> GameProjectUtils::WarningProjectNameNotification = NULL;
 
 bool GameProjectUtils::IsValidProjectFileForCreation(const FString& ProjectFile, FText& OutFailReason)
 {
@@ -137,14 +138,6 @@ bool GameProjectUtils::OpenProject(const FString& ProjectFile, FText& OutFailRea
 		return false;
 	}
 
-	if ( BaseProjectFile.Len() > MAX_PROJECT_NAME_LENGTH )
-	{
-		FFormatNamedArguments Args;
-		Args.Add( TEXT("MaxProjectNameLength"), MAX_PROJECT_NAME_LENGTH );
-		OutFailReason = FText::Format( LOCTEXT( "ProjectNameTooLong", "Project names must not be longer than {MaxProjectNameLength} characters." ), Args );
-		return false;
-	}
-
 	const int32 MaxProjectPathLength = PLATFORM_MAX_FILEPATH_LENGTH - MAX_PROJECT_PATH_BUFFER_SPACE;
 	if ( FPaths::GetBaseFilename(ProjectFile, false).Len() > MaxProjectPathLength )
 	{
@@ -405,6 +398,53 @@ void GameProjectUtils::CheckForOutOfDateGameProjectFile()
 	}
 }
 
+void GameProjectUtils::CheckAndWarnProjectFilenameValid()
+{
+	const FString& LoadedProjectFilePath = FPaths::IsProjectFilePathSet() ? FPaths::GetProjectFilePath() : FString();
+	if ( !LoadedProjectFilePath.IsEmpty() )
+	{
+		const FString BaseProjectFile = FPaths::GetBaseFilename(LoadedProjectFilePath);
+		if ( BaseProjectFile.Len() > MAX_PROJECT_NAME_LENGTH )
+		{
+			FFormatNamedArguments Args;
+			Args.Add( TEXT("MaxProjectNameLength"), MAX_PROJECT_NAME_LENGTH );
+			const FText WarningReason = FText::Format( LOCTEXT( "WarnProjectNameTooLong", "Project names must not be longer than {MaxProjectNameLength} characters.\nYou might have problems saving or modifying a project with a longer name." ), Args );
+			const FText WarningReasonOkText = LOCTEXT("WarningReasonOkText", "Ok");
+
+			FNotificationInfo Info(WarningReason);
+			Info.bFireAndForget = false;
+			Info.bUseLargeFont = false;
+			Info.bUseThrobber = false;
+			Info.bUseSuccessFailIcons = false;
+			Info.FadeOutDuration = 3.f;
+			Info.ButtonDetails.Add(FNotificationButtonInfo(WarningReasonOkText, FText(), FSimpleDelegate::CreateStatic(&GameProjectUtils::OnWarningReasonOk)));
+
+			if (WarningProjectNameNotification.IsValid())
+			{
+				WarningProjectNameNotification.Pin()->ExpireAndFadeout();
+				WarningProjectNameNotification.Reset();
+			}
+
+			WarningProjectNameNotification = FSlateNotificationManager::Get().AddNotification(Info);
+
+			if (WarningProjectNameNotification.IsValid())
+			{
+				WarningProjectNameNotification.Pin()->SetCompletionState(SNotificationItem::CS_Pending);
+			}
+		}
+	}
+}
+
+void GameProjectUtils::OnWarningReasonOk()
+{
+	if ( WarningProjectNameNotification.IsValid() )
+	{
+		WarningProjectNameNotification.Pin()->SetCompletionState(SNotificationItem::CS_None);
+		WarningProjectNameNotification.Pin()->ExpireAndFadeout();
+		WarningProjectNameNotification.Reset();
+	}
+}
+
 bool GameProjectUtils::UpdateGameProject(const FString& EngineIdentifier)
 {
 	const FString& ProjectFilename = FPaths::IsProjectFilePathSet() ? FPaths::GetProjectFilePath() : FString();
diff --git a/Engine/Source/Editor/GameProjectGeneration/Private/GameProjectUtils.h b/Engine/Source/Editor/GameProjectGeneration/Private/GameProjectUtils.h
--- a/Engine/Source/Editor/GameProjectGeneration/Private/GameProjectUtils.h
+++ b/Engine/Source/Editor/GameProjectGeneration/Private/GameProjectUtils.h
@@ -23,6 +23,9 @@ public:
 	/** Prompts the user to update his project file, if necessary. */
 	static void CheckForOutOfDateGameProjectFile();
 
+	/** Warn the user if the project filename is invalid in case they renamed it outside the editor */
+	static void CheckAndWarnProjectFilenameValid();
+
 	/** Updates the currently loaded project. Returns true if the project was updated successfully or if no update was needed */
 	static bool UpdateGameProject(const FString &EngineIdentifier);
 
@@ -218,6 +221,10 @@ private:
 	/** Internal handler for AddCodeToProject*/
 	static bool AddCodeToProject_Internal(const FString& NewClassName, const FString& NewClassPath, const UClass* ParentClass, FString& OutHeaderFilePath, FString& OutCppFilePath, FText& OutFailReason);
 
+	/** Handler for the user confirming they've read the name legnth warning */
+	static void OnWarningReasonOk();
+
 private:
 	static TWeakPtr<SNotificationItem> UpdateGameProjectNotification;
+	static TWeakPtr<SNotificationItem> WarningProjectNameNotification;
 };
diff --git a/Engine/Source/Editor/GameProjectGeneration/Public/GameProjectGenerationModule.h b/Engine/Source/Editor/GameProjectGeneration/Public/GameProjectGenerationModule.h
--- a/Engine/Source/Editor/GameProjectGeneration/Public/GameProjectGenerationModule.h
+++ b/Engine/Source/Editor/GameProjectGeneration/Public/GameProjectGenerationModule.h
@@ -6,7 +6,7 @@
 #include "ModuleInterface.h"
 
 /**
- * Content browser module
+ * Game Project Generation module
  */
 class FGameProjectGenerationModule : public IModuleInterface
 {
@@ -62,6 +62,9 @@ public:
 	/** Update code resource files */
 	virtual bool UpdateCodeResourceFiles(TArray<FString>& OutCreatedFiles, FText& OutFailReason);
 
+	/** Warn the user if the project filename is invalid in case they renamed it outside the editor */
+	virtual void CheckAndWarnProjectFilenameValid();
+
 private:
 	FAddCodeToProjectDialogOpenedEvent AddCodeToProjectDialogOpenedEvent;
 
diff --git a/Engine/Source/Editor/UnrealEd/Private/UnrealEd.cpp b/Engine/Source/Editor/UnrealEd/Private/UnrealEd.cpp
--- a/Engine/Source/Editor/UnrealEd/Private/UnrealEd.cpp
+++ b/Engine/Source/Editor/UnrealEd/Private/UnrealEd.cpp
@@ -103,6 +103,7 @@ int32 EditorInit( IEngineLoop& EngineLoop )
 	if ( FPaths::IsProjectFilePathSet() )
 	{
 		FGameProjectGenerationModule::Get().CheckForOutOfDateGameProjectFile();
+		FGameProjectGenerationModule::Get().CheckAndWarnProjectFilenameValid();
 	}
 
 	// =================== EDITOR STARTUP FINISHED ===================
