Project template names and descriptions now fallback to English if a valid translation can't be found for the current language

TTP# 333732 - EDITOR: The new project wizard falls back to using TP_TemplateName if it can't find a matching translation, even if an English translation is available

diff --git a/Engine/Source/Editor/GameProjectGeneration/Classes/TemplateProjectDefs.h b/Engine/Source/Editor/GameProjectGeneration/Classes/TemplateProjectDefs.h
--- a/Engine/Source/Editor/GameProjectGeneration/Classes/TemplateProjectDefs.h
+++ b/Engine/Source/Editor/GameProjectGeneration/Classes/TemplateProjectDefs.h
@@ -75,10 +75,10 @@ class UTemplateProjectDefs : public UObject
 	/** Fixes up all strings in this definitions object to replace %TEMPLATENAME% with the supplied template name and %PROJECTNAME% with the supplied project name */
 	void FixupStrings(const FString& TemplateName, const FString& ProjectName);
 
-	/** Returns the display name for the current culture */
+	/** Returns the display name for the current culture, or English if the current culture has no translation */
 	FText GetDisplayNameText();
 
-	/** Returns the display name for the current culture */
+	/** Returns the display name for the current culture, or English if the current culture has no translation */
 	FText GetLocalizedDescription();
 
 private:
diff --git a/Engine/Source/Editor/GameProjectGeneration/Private/TemplateProjectDefs.cpp b/Engine/Source/Editor/GameProjectGeneration/Private/TemplateProjectDefs.cpp
--- a/Engine/Source/Editor/GameProjectGeneration/Private/TemplateProjectDefs.cpp
+++ b/Engine/Source/Editor/GameProjectGeneration/Private/TemplateProjectDefs.cpp
@@ -3,6 +3,45 @@
 
 #include "GameProjectGenerationPrivatePCH.h"
 
+namespace
+{
+
+FText GetLocalizedText(const TArray<FLocalizedTemplateString>& LocalizedStrings)
+{
+	const FString DefaultCultureISOLanguageName = "en";
+	const FString CurrentCultureISOLanguageName = FInternationalization::Get().GetCurrentCulture()->GetTwoLetterISOLanguageName();
+	FText FallbackText;
+
+	for ( const FLocalizedTemplateString& LocalizedString : LocalizedStrings )
+	{
+		if ( LocalizedString.Language == CurrentCultureISOLanguageName )
+		{
+			return FText::FromString(LocalizedString.Text);
+		}
+
+		if ( LocalizedString.Language == DefaultCultureISOLanguageName )
+		{
+			FallbackText = FText::FromString(LocalizedString.Text);
+		}
+	}
+
+	// Did we find an English fallback?
+	if ( !FallbackText.IsEmpty() )
+	{
+		return FallbackText;
+	}
+
+	// We failed to find English, see if we have any translations available to use
+	if ( LocalizedStrings.Num() )
+	{
+		return FText::FromString(LocalizedStrings[0].Text);
+	}
+
+	return FText();
+}
+
+}
+
 UTemplateProjectDefs::UTemplateProjectDefs(const class FPostConstructInitializeProperties& PCIP)
 	: Super(PCIP)
 {
@@ -45,32 +84,12 @@ void UTemplateProjectDefs::FixupStrings(const FString& TemplateName, const FStri
 
 FText UTemplateProjectDefs::GetDisplayNameText()
 {
-	const TSharedRef< FCulture > CurrentCulture = FInternationalization::Get().GetCurrentCulture();
-	for ( auto NameIt = LocalizedDisplayNames.CreateConstIterator(); NameIt; ++NameIt )
-	{
-		const FLocalizedTemplateString& Name = *NameIt;
-		if ( Name.Language == CurrentCulture->GetTwoLetterISOLanguageName() )
-		{
-			return FText::FromString(Name.Text);
-		}
-	}
-
-	return FText();
+	return GetLocalizedText(LocalizedDisplayNames);
 }
 
 FText UTemplateProjectDefs::GetLocalizedDescription()
 {
-	const TSharedRef< FCulture >  CurrentCulture = FInternationalization::Get().GetCurrentCulture();
-	for ( auto DescriptionIt = LocalizedDescriptions.CreateConstIterator(); DescriptionIt; ++DescriptionIt )
-	{
-		const FLocalizedTemplateString& Description = *DescriptionIt;
-		if ( Description.Language == CurrentCulture->GetTwoLetterISOLanguageName() )
-		{
-			return FText::FromString(Description.Text);
-		}
-	}
-
-	return FText();
+	return GetLocalizedText(LocalizedDescriptions);
 }
 
 void UTemplateProjectDefs::FixString(FString& InOutStringToFix, const FString& TemplateName, const FString& ProjectName)
