Out of date projects now appear faded-out on the project browser

This is to highlight that they will need to be upgraded before they can be used in the current version of the engine. Currently there is no upgrade step, but this will be implemented shortly. This addresses some of TTP #333190 (Project Browser: Improved workflow for opening projects from earlier versions of the editor)

diff --git a/Engine/Source/Editor/EditorStyle/Private/SlateEditorStyle.cpp b/Engine/Source/Editor/EditorStyle/Private/SlateEditorStyle.cpp
--- a/Engine/Source/Editor/EditorStyle/Private/SlateEditorStyle.cpp
+++ b/Engine/Source/Editor/EditorStyle/Private/SlateEditorStyle.cpp
@@ -4406,7 +4406,12 @@ void FSlateEditorStyle::FStyle::SetupToolkitStyles()
 	{
 		Set("ProjectBrowser.Tab.Text", FTextBlockStyle(NormalText)
 			.SetFont(TTF_CORE_FONT("Fonts/Roboto-Regular", 24))
-			.SetShadowOffset(FVector2D(0,1)));
+			.SetShadowOffset(FVector2D(0, 1)));
+
+		Set("ProjectBrowser.VersionOverlayText", FTextBlockStyle(NormalText)
+			.SetFont(TTF_CORE_FONT("Fonts/Roboto-Bold", 14))
+			.SetShadowOffset(FVector2D(0, 1)));
+
 
 		Set("ProjectBrowser.Background", new BOX_BRUSH("Common/FlatColorSquare", FVector2D(1.0f, 1.0f), FMargin(0), FLinearColor(FColor(0xff404040))));
 		Set("ProjectBrowser.Tab.ActiveBackground", new BOX_BRUSH("Common/FlatColorSquare", FVector2D(1.0f, 1.0f), FMargin(0), FLinearColor(FColor(0xff404040))));
diff --git a/Engine/Source/Editor/GameProjectGeneration/Private/SProjectBrowser.cpp b/Engine/Source/Editor/GameProjectGeneration/Private/SProjectBrowser.cpp
--- a/Engine/Source/Editor/GameProjectGeneration/Private/SProjectBrowser.cpp
+++ b/Engine/Source/Editor/GameProjectGeneration/Private/SProjectBrowser.cpp
@@ -15,17 +15,27 @@ struct FProjectItem
 {
 	FText Name;
 	FText Description;
+	FEngineVersion EngineVersion;
+	bool bUpToDate;
 	FString ProjectFile;
 	TSharedPtr<FSlateBrush> ProjectThumbnail;
 	bool bIsNewProjectItem;
 
-	FProjectItem(const FText& InName, const FText& InDescription, const TSharedPtr<FSlateBrush>& InProjectThumbnail, const FString& InProjectFile, bool InIsNewProjectItem)
+	FProjectItem(const FText& InName, const FText& InDescription, const FEngineVersion& InEngineVersion, bool InUpToDate, const TSharedPtr<FSlateBrush>& InProjectThumbnail, const FString& InProjectFile, bool InIsNewProjectItem)
 		: Name(InName)
 		, Description(InDescription)
+		, EngineVersion(InEngineVersion)
+		, bUpToDate(InUpToDate)
 		, ProjectFile(InProjectFile)
 		, ProjectThumbnail(InProjectThumbnail)
 		, bIsNewProjectItem(InIsNewProjectItem)
 	{ }
+
+	/** Check if this project is up to date */
+	bool IsUpToDate() const
+	{
+		return bUpToDate || EngineVersion.IsEmpty();
+	}
 };
 
 
@@ -312,10 +322,14 @@ TSharedRef<ITableRow> SProjectBrowser::MakeProjectViewWidget(TSharedPtr<FProject
 	}
 	else
 	{
+		const FLinearColor Tint = ProjectItem->IsUpToDate() ? FLinearColor::White : FLinearColor::White.CopyWithNewOpacity(0.5);
+
 		// Drop shadow border
 		Thumbnail =	SNew(SBorder)
 			.Padding(ThumbnailBorderPadding)
 			.BorderImage( FEditorStyle::GetBrush("ContentBrowser.ThumbnailShadow") )
+			.ColorAndOpacity(Tint)
+			.BorderBackgroundColor(Tint)
 			[
 				SNew(SImage).Image(this, &SProjectBrowser::GetProjectItemImage, TWeakPtr<FProjectItem>(ProjectItem))
 			];
@@ -337,7 +351,25 @@ TSharedRef<ITableRow> SProjectBrowser::MakeProjectViewWidget(TSharedPtr<FProject
 				.WidthOverride(ThumbnailSize + ThumbnailBorderPadding * 2)
 				.HeightOverride(ThumbnailSize + ThumbnailBorderPadding * 2)
 				[
-					Thumbnail.ToSharedRef()
+					SNew(SOverlay)
+
+					+ SOverlay::Slot()
+					[
+						Thumbnail.ToSharedRef()
+					]
+
+					// Show the out of date engine version for this project file
+					+ SOverlay::Slot()
+					.HAlign(HAlign_Right)
+					.VAlign(VAlign_Bottom)
+					.Padding(10)
+					[
+						SNew(STextBlock)
+						.Text(ProjectItem->EngineVersion.ToString(EVersionComponent::Minor))
+						.TextStyle(FEditorStyle::Get(), "ProjectBrowser.VersionOverlayText")
+						.ColorAndOpacity(FLinearColor::White.CopyWithNewOpacity(0.5f))
+						.Visibility(ProjectItem->IsUpToDate() ? EVisibility::Collapsed : EVisibility::Visible)
+					]
 				]
 			]
 
@@ -374,6 +406,14 @@ TSharedRef<SToolTip> SProjectBrowser::MakeProjectToolTip( TSharedPtr<FProjectIte
 		AddToToolTipInfoBox( InfoBox, LOCTEXT("ProjectTileTooltipPath", "Path"), FText::FromString(ProjectPath) );
 	}
 
+	if (!ProjectItem->IsUpToDate())
+	{
+		FFormatOrderedArguments Args;
+		Args.Add(FText::FromString(ProjectItem->EngineVersion.ToString(EVersionComponent::Minor)));
+		AddToToolTipInfoBox(InfoBox, LOCTEXT("ProjectVersion", "Version"),
+			FText::Format(LOCTEXT("ProjectTileOutOfDate", "This project was created with engine version {0}. It may require an update to work with this build."), Args));
+	}
+
 	TSharedRef<SToolTip> Tooltip = SNew(SToolTip)
 	.TextMargin(1)
 	.BorderImage( FEditorStyle::GetBrush("ProjectBrowser.TileViewTooltip.ToolTipBorder") )
@@ -651,7 +691,7 @@ FReply SProjectBrowser::FindProjects()
 				}
 
 				const bool bIsNewProjectItem = false;
-				TSharedRef<FProjectItem> NewProjectItem = MakeShareable( new FProjectItem(ProjectName, ProjectDescription, DynamicBrush, ProjectFilename, bIsNewProjectItem ) );
+				TSharedRef<FProjectItem> NewProjectItem = MakeShareable( new FProjectItem(ProjectName, ProjectDescription, ProjectStatus.EngineVersion, ProjectStatus.bUpToDate, DynamicBrush, ProjectFilename, bIsNewProjectItem ) );
 				AddProjectToCategory(NewProjectItem, ProjectCategory);
 			}
 		}
diff --git a/Engine/Source/Runtime/Core/Private/Misc/EngineVersion.cpp b/Engine/Source/Runtime/Core/Private/Misc/EngineVersion.cpp
--- a/Engine/Source/Runtime/Core/Private/Misc/EngineVersion.cpp
+++ b/Engine/Source/Runtime/Core/Private/Misc/EngineVersion.cpp
@@ -141,31 +141,34 @@ bool FEngineVersion::Parse(const FString &Text, FEngineVersion &OutVersion)
 
 EVersionComparison::Type FEngineVersion::GetNewest(const FEngineVersion &First, const FEngineVersion &Second, EVersionComponent::Type *OutComponent)
 {
+	EVersionComponent::Type LocalComponent;
+	auto& Component = OutComponent ? *OutComponent : LocalComponent;
+
 	// Compare major versions
 	if(First.Major != Second.Major)
 	{
-		*OutComponent = EVersionComponent::Major;
+		Component = EVersionComponent::Major;
 		return (First.Major > Second.Major)? EVersionComparison::First : EVersionComparison::Second;
 	}
 
 	// Compare minor versions
 	if(First.Minor != Second.Minor)
 	{
-		*OutComponent = EVersionComponent::Minor;
+		Component = EVersionComponent::Minor;
 		return (First.Minor > Second.Minor)? EVersionComparison::First : EVersionComparison::Second;
 	}
 
 	// Compare patch versions
 	if(First.Patch != Second.Patch)
 	{
-		*OutComponent = EVersionComponent::Patch;
+		Component = EVersionComponent::Patch;
 		return (First.Patch > Second.Patch)? EVersionComparison::First : EVersionComparison::Second;
 	}
 
 	// Compare changelists (only if they're both from the same vendor)
 	if(First.IsLicenseeVersion() == Second.IsLicenseeVersion() && First.GetChangelist() != Second.GetChangelist())
 	{
-		*OutComponent = EVersionComponent::Changelist;
+		Component = EVersionComponent::Changelist;
 		return (First.GetChangelist() > Second.GetChangelist()) ? EVersionComparison::First : EVersionComparison::Second;
 	}
 
diff --git a/Engine/Source/Runtime/Projects/Private/ProjectManager.cpp b/Engine/Source/Runtime/Projects/Private/ProjectManager.cpp
--- a/Engine/Source/Runtime/Projects/Private/ProjectManager.cpp
+++ b/Engine/Source/Runtime/Projects/Private/ProjectManager.cpp
@@ -322,6 +322,7 @@ bool FProjectManager::QueryStatusForProject(const FString& FilePath, const FStri
 	OutProjectStatus.Name = ProjectInfo.Name;
 	OutProjectStatus.Description = ProjectInfo.Description;
 	OutProjectStatus.Category = ProjectInfo.Category;
+	OutProjectStatus.EngineVersion = ProjectInfo.EngineVersion;
 	OutProjectStatus.bSignedSampleProject = NewProject->IsSignedSampleProject(FilePath);
 	OutProjectStatus.bUpToDate = NewProject->IsUpToDate(EngineIdentifier);
 
diff --git a/Engine/Source/Runtime/Projects/Public/Interfaces/IProjectManager.h b/Engine/Source/Runtime/Projects/Public/Interfaces/IProjectManager.h
--- a/Engine/Source/Runtime/Projects/Public/Interfaces/IProjectManager.h
+++ b/Engine/Source/Runtime/Projects/Public/Interfaces/IProjectManager.h
@@ -4,6 +4,7 @@
 
 
 #include "IPluginManagerShared.h"
+#include "EngineVersion.h"
 
 /**
  * Simple data structure that is filled when querying information about projects
@@ -21,6 +22,9 @@ public:
 	/** The UI category of this project */
 	FString Category;
 
+	/** The version of the engine that this project is made for */
+	FEngineVersion EngineVersion;
+
 	/** True if this project is a sample provided by epic */
 	bool bSignedSampleProject;
 
