Don't allow an underscore in the project name if the XB1 XDK is installed
Fixes TTP333191
#ue4

diff --git a/Engine/Source/Editor/GameProjectGeneration/GameProjectGeneration.Build.cs b/Engine/Source/Editor/GameProjectGeneration/GameProjectGeneration.Build.cs
--- a/Engine/Source/Editor/GameProjectGeneration/GameProjectGeneration.Build.cs
+++ b/Engine/Source/Editor/GameProjectGeneration/GameProjectGeneration.Build.cs
@@ -28,7 +28,8 @@ public class GameProjectGeneration : ModuleRules
 				"SlateCore",
                 "EditorStyle",
                 "SourceControl",
-                "UnrealEd",
+ 				"TargetPlatform",
+				"UnrealEd",
 				"DesktopPlatform"
 			}
 		);
diff --git a/Engine/Source/Editor/GameProjectGeneration/Private/GameProjectUtils.cpp b/Engine/Source/Editor/GameProjectGeneration/Private/GameProjectUtils.cpp
--- a/Engine/Source/Editor/GameProjectGeneration/Private/GameProjectUtils.cpp
+++ b/Engine/Source/Editor/GameProjectGeneration/Private/GameProjectUtils.cpp
@@ -11,6 +11,7 @@
 #include "EngineBuildSettings.h"
 
 #include "DesktopPlatformModule.h"
+#include "TargetPlatform.h"
 
 #define LOCTEXT_NAMESPACE "GameProjectUtils"
 
@@ -84,6 +85,12 @@ bool GameProjectUtils::IsValidProjectFileForCreation(const FString& ProjectFile,
 		return false;
 	}
 
+	if (NameContainsUnderscoreAndXB1Installed(BaseProjectFile))
+	{
+		OutFailReason = LOCTEXT( "ProjectNameContainsIllegalCharactersOnXB1", "Project names may not contain an underscore when the Xbox One XDK is installed." );
+		return false;
+	}
+
 	if ( !FPaths::ValidatePath(FPaths::GetPath(ProjectFile), &OutFailReason) )
 	{
 		return false;
@@ -164,6 +171,12 @@ bool GameProjectUtils::OpenProject(const FString& ProjectFile, FText& OutFailRea
 		return false;
 	}
 
+	if (NameContainsUnderscoreAndXB1Installed(BaseProjectFile))
+	{
+		OutFailReason = LOCTEXT( "ProjectNameContainsIllegalCharactersOnXB1", "Project names may not contain an underscore when the Xbox One XDK is installed." );
+		return false;
+	}
+
 	if ( !FPaths::ValidatePath(FPaths::GetPath(ProjectFile), &OutFailReason) )
 	{
 		return false;
@@ -1089,6 +1102,31 @@ bool GameProjectUtils::NameContainsOnlyLegalCharacters(const FString& TestName,
 	return !bContainsIllegalCharacters;
 }
 
+bool GameProjectUtils::NameContainsUnderscoreAndXB1Installed(const FString& TestName)
+{
+	bool bContainsIllegalCharacters = false;
+
+	// Only allow alphanumeric characters in the project name
+	for ( int32 CharIdx = 0 ; CharIdx < TestName.Len() ; ++CharIdx )
+	{
+		const FString& Char = TestName.Mid( CharIdx, 1 );
+		if ( Char == TEXT("_") )
+		{
+			const ITargetPlatform* Platform = GetTargetPlatformManager()->FindTargetPlatform(TEXT("XboxOne"));
+			if (Platform)
+			{
+				FString NotInstalledDocLink;
+				if (Platform->IsSdkInstalled(true, NotInstalledDocLink))
+				{
+					bContainsIllegalCharacters = true;
+				}
+			}
+		}
+	}
+
+	return bContainsIllegalCharacters;
+}
+
 bool GameProjectUtils::ProjectFileExists(const FString& ProjectFile)
 {
 	return FPlatformFileManager::Get().GetPlatformFile().FileExists(*ProjectFile);
diff --git a/Engine/Source/Editor/GameProjectGeneration/Private/GameProjectUtils.h b/Engine/Source/Editor/GameProjectGeneration/Private/GameProjectUtils.h
--- a/Engine/Source/Editor/GameProjectGeneration/Private/GameProjectUtils.h
+++ b/Engine/Source/Editor/GameProjectGeneration/Private/GameProjectUtils.h
@@ -111,6 +111,9 @@ private:
 	/** Checks the name for illegal characters */
 	static bool NameContainsOnlyLegalCharacters(const FString& TestName, FString& OutIllegalCharacters);
 
+	/** Checks the name for an underscore and the existence of XB1 XDK */
+	static bool NameContainsUnderscoreAndXB1Installed(const FString& TestName);
+
 	/** Returns true if the project file exists on disk */
 	static bool ProjectFileExists(const FString& ProjectFile);
 
