Fix typo in GameProjectUtils::AddCodeToProject_Internal
https://github.com/EpicGames/UnrealEngine/pull/143/files (mlewand)

diff --git a/Engine/Source/Editor/GameProjectGeneration/Private/GameProjectUtils.cpp b/Engine/Source/Editor/GameProjectGeneration/Private/GameProjectUtils.cpp
--- a/Engine/Source/Editor/GameProjectGeneration/Private/GameProjectUtils.cpp
+++ b/Engine/Source/Editor/GameProjectGeneration/Private/GameProjectUtils.cpp
@@ -2089,7 +2089,7 @@ bool GameProjectUtils::AddCodeToProject_Internal(const FString& NewClassName, co
 
 	if ( !FApp::HasGameName() )
 	{
-		OutFailReason = LOCTEXT("AddCodeToProject_NoGameName", "You can not add code because you have not loaded a projet.");
+		OutFailReason = LOCTEXT("AddCodeToProject_NoGameName", "You can not add code because you have not loaded a project.");
 		return false;
 	}
 
