Change project 'upgrades' to be based on the format of the project file and nothing to do with engine version or association, and allow the UpdateGameProject commandlet to always write latest version information to the file.

diff --git a/Engine/Source/Developer/DesktopPlatform/Private/DesktopPlatformBase.cpp b/Engine/Source/Developer/DesktopPlatform/Private/DesktopPlatformBase.cpp
--- a/Engine/Source/Developer/DesktopPlatform/Private/DesktopPlatformBase.cpp
+++ b/Engine/Source/Developer/DesktopPlatform/Private/DesktopPlatformBase.cpp
@@ -3,11 +3,18 @@
 #include "DesktopPlatformPrivatePCH.h"
 #include "DesktopPlatformBase.h"
 #include "UProjectInfo.h"
+#include "EngineVersion.h"
+#include "ModuleManager.h"
+
+#define LOCTEXT_NAMESPACE "DesktopPlatform"
 
 FString FDesktopPlatformBase::GetCurrentEngineIdentifier()
 {
-	FString Identifier;
-	return GetEngineIdentifierFromRootDir(FPlatformMisc::RootDir(), Identifier) ? Identifier : TEXT("");
+	if(CurrentEngineIdentifier.Len() == 0 && !GetEngineIdentifierFromRootDir(FPlatformMisc::RootDir(), CurrentEngineIdentifier))
+	{
+		CurrentEngineIdentifier.Empty();
+	}
+	return CurrentEngineIdentifier;
 }
 
 void FDesktopPlatformBase::EnumerateLauncherEngineInstallations(TMap<FString, FString> &OutInstallations)
@@ -121,6 +128,11 @@ bool FDesktopPlatformBase::IsPreferredEngineIdentifier(const FString &Identifier
 	}
 }
 
+bool FDesktopPlatformBase::IsStockEngineRelease(const FString &Identifier)
+{
+	return Identifier.Len() > 0 && FChar::IsDigit(Identifier[0]);
+}
+
 bool FDesktopPlatformBase::IsSourceDistribution(const FString &EngineRootDir)
 {
 	// Check for the existence of a SourceBuild.txt file
@@ -160,7 +172,7 @@ bool FDesktopPlatformBase::SetEngineIdentifierForProject(const FString &ProjectF
 		FString RootDir;
 		if(GetEngineRootDirFromIdentifier(Identifier, RootDir))
 		{
-			FUProjectDictionary Dictionary(RootDir);
+			const FUProjectDictionary &Dictionary = GetCachedProjectDictionary(RootDir);
 			if(!Dictionary.IsForeignProject(ProjectFileName))
 			{
 				Identifier.Empty();
@@ -198,16 +210,96 @@ bool FDesktopPlatformBase::GetEngineIdentifierForProject(const FString &ProjectF
 	{
 		ParentDir.RemoveAt(SeparatorIdx, ParentDir.Len() - SeparatorIdx);
 
-		FUProjectDictionary Dictionary(ParentDir);
+		const FUProjectDictionary &Dictionary = GetCachedProjectDictionary(ParentDir);
 		if(!Dictionary.IsForeignProject(ProjectFileName) && GetEngineIdentifierFromRootDir(ParentDir, OutIdentifier))
 		{
 			return true;
 		}
 	}
 
+	// Otherwise check the engine version string for 4.0, in case this project existed before the engine association stuff went in
+	FString EngineVersionString = ProjectFile->GetStringField(TEXT("EngineVersion"));
+	if(EngineVersionString.Len() > 0)
+	{
+		FEngineVersion EngineVersion;
+		if(FEngineVersion::Parse(EngineVersionString, EngineVersion) && EngineVersion.IsPromotedBuild() && EngineVersion.ToString(EVersionComponent::Minor) == TEXT("4.0"))
+		{
+			OutIdentifier = TEXT("4.0");
+			return true;
+		}
+	}
+
 	return false;
 }
 
+bool FDesktopPlatformBase::CompileGameProject(const FString& RootDir, const FString& ProjectFileName, FFeedbackContext* Warn)
+{
+	// Get the target name
+	FString Arguments = FString::Printf(TEXT("%sEditor %s %s"), *FPaths::GetBaseFilename(ProjectFileName), FModuleManager::Get().GetUBTConfiguration(), FPlatformMisc::GetUBTPlatform());
+
+	// Append the project name if it's a foreign project
+	if ( !ProjectFileName.IsEmpty() )
+	{
+		FUProjectDictionary ProjectDictionary(RootDir);
+		if(ProjectDictionary.IsForeignProject(ProjectFileName))
+		{
+			Arguments += FString::Printf(TEXT(" -project=\"%s\""), *IFileManager::Get().ConvertToAbsolutePathForExternalAppForRead(*ProjectFileName));
+		}
+	}
+
+	// Append the Rocket flag
+	if(!IsSourceDistribution(RootDir))
+	{
+		Arguments += TEXT(" -rocket");
+	}
+
+	// Append any other options
+	Arguments += " -editorrecompile -progress";
+
+	// Run UBT
+	return RunUnrealBuildTool(LOCTEXT("CompilingProject", "Compiling project..."), RootDir, Arguments, Warn);
+}
+
+bool FDesktopPlatformBase::GenerateProjectFiles(const FString& RootDir, const FString& ProjectFileName, FFeedbackContext* Warn)
+{
+#if PLATFORM_MAC
+	FString Arguments = TEXT("-xcodeprojectfile");
+#else
+	FString Arguments = TEXT("-projectfiles");
+#endif
+
+	// Build the arguments to pass to UBT
+	if ( !ProjectFileName.IsEmpty() )
+	{
+		// Figure out whether it's a foreign project
+		FUProjectDictionary ProjectDictionary(RootDir);
+		if(ProjectDictionary.IsForeignProject(ProjectFileName))
+		{
+			Arguments += FString::Printf(TEXT(" -project=\"%s\""), *IFileManager::Get().ConvertToAbsolutePathForExternalAppForRead(*ProjectFileName));
+		}
+
+		// Always include game source
+		Arguments += " -game";
+
+		// Determine whether or not to include engine source
+		if(IsSourceDistribution(RootDir))
+		{
+			Arguments += " -engine";
+		}
+		else
+		{
+			Arguments += " -rocket";
+		}
+	}
+	Arguments += " -progress";
+
+	// Run UnrealBuildTool
+	Warn->BeginSlowTask(LOCTEXT("GeneratingProjectFiles", "Generating project files..."), true, true);
+	bool bRes = RunUnrealBuildTool(LOCTEXT("GeneratingProjectFiles", "Generating project files..."), RootDir, Arguments, Warn);
+	Warn->EndSlowTask();
+	return bRes;
+}
+
 bool FDesktopPlatformBase::ReadLauncherInstallationList(TMap<FString, FString> &OutInstallations)
 {
 	FString InstalledListFile = FString(FPlatformProcess::ApplicationSettingsDir()) / TEXT("UnrealEngineLauncher/LauncherInstalled.dat");
@@ -322,3 +414,18 @@ bool FDesktopPlatformBase::SaveProjectFile(const FString &FileName, TSharedPtr<F
 
 	return true;
 }
+
+const FUProjectDictionary &FDesktopPlatformBase::GetCachedProjectDictionary(const FString& RootDir)
+{
+	FString NormalizedRootDir = RootDir;
+	FPaths::NormalizeDirectoryName(NormalizedRootDir);
+
+	FUProjectDictionary *Dictionary = CachedProjectDictionaries.Find(NormalizedRootDir);
+	if(Dictionary == NULL)
+	{
+		Dictionary = &CachedProjectDictionaries.Add(RootDir, FUProjectDictionary(RootDir));
+	}
+	return *Dictionary;
+}
+
+#undef LOCTEXT_NAMESPACE
diff --git a/Engine/Source/Developer/DesktopPlatform/Private/DesktopPlatformBase.h b/Engine/Source/Developer/DesktopPlatform/Private/DesktopPlatformBase.h
--- a/Engine/Source/Developer/DesktopPlatform/Private/DesktopPlatformBase.h
+++ b/Engine/Source/Developer/DesktopPlatform/Private/DesktopPlatformBase.h
@@ -4,6 +4,7 @@
 
 #include "IDesktopPlatform.h"
 #include "Runtime/Core/Public/Serialization/Json/Json.h"
+#include "UProjectInfo.h"
 
 class FDesktopPlatformBase : public IDesktopPlatform
 {
@@ -19,6 +20,7 @@ public:
 	virtual bool GetDefaultEngineRootDir(FString &OutRootDir) OVERRIDE;
 	virtual bool IsPreferredEngineIdentifier(const FString &Identifier, const FString &OtherIdentifier);
 
+	virtual bool IsStockEngineRelease(const FString &Identifier) OVERRIDE;
 	virtual bool IsSourceDistribution(const FString &RootDir) OVERRIDE;
 	virtual bool IsPerforceBuild(const FString &RootDir) OVERRIDE;
 	virtual bool IsValidRootDirectory(const FString &RootDir) OVERRIDE;
@@ -26,11 +28,19 @@ public:
 	virtual bool SetEngineIdentifierForProject(const FString &ProjectFileName, const FString &Identifier) OVERRIDE;
 	virtual bool GetEngineIdentifierForProject(const FString &ProjectFileName, FString &OutIdentifier) OVERRIDE;
 
+	virtual bool CompileGameProject(const FString& RootDir, const FString& ProjectFileName, FFeedbackContext* Warn) OVERRIDE;
+	virtual bool GenerateProjectFiles(const FString& RootDir, const FString& ProjectFileName, FFeedbackContext* Warn) OVERRIDE;
+
 private:
+	FString CurrentEngineIdentifier;
+	TMap<FString, FUProjectDictionary> CachedProjectDictionaries;
+
 	bool ReadLauncherInstallationList(TMap<FString, FString> &OutInstallations);
 	void CheckForLauncherEngineInstallation(const FString &AppId, const FString &Identifier, TMap<FString, FString> &OutInstallations);
 	int32 ParseReleaseVersion(const FString &Version);
 
 	TSharedPtr<FJsonObject> LoadProjectFile(const FString &FileName);
 	bool SaveProjectFile(const FString &FileName, TSharedPtr<FJsonObject> Object);
+
+	const FUProjectDictionary &GetCachedProjectDictionary(const FString& RootDir);
 };
diff --git a/Engine/Source/Developer/DesktopPlatform/Private/Mac/DesktopPlatformMac.cpp b/Engine/Source/Developer/DesktopPlatform/Private/Mac/DesktopPlatformMac.cpp
--- a/Engine/Source/Developer/DesktopPlatform/Private/Mac/DesktopPlatformMac.cpp
+++ b/Engine/Source/Developer/DesktopPlatform/Private/Mac/DesktopPlatformMac.cpp
@@ -2,6 +2,7 @@
 
 #include "DesktopPlatformPrivatePCH.h"
 #include "MacApplication.h"
+#include "FeedbackContextMarkup.h"
 
 #define LOCTEXT_NAMESPACE "DesktopPlatform"
 
@@ -619,4 +620,23 @@ bool FDesktopPlatformMac::UpdateFileAssociations()
 	return Status == noErr;
 }
 
+bool FDesktopPlatformMac::RunUnrealBuildTool(const FText& Description, const FString& RootDir, const FString& Arguments, FFeedbackContext* Warn)
+{
+	// Get the path to UBT
+	FString UnrealBuildToolPath = RootDir / TEXT("Engine/Binaries/DotNET/UnrealBuildTool.exe");
+	if(IFileManager::Get().FileSize(*UnrealBuildToolPath) < 0)
+	{
+		Warn->Logf(ELogVerbosity::Error, TEXT("Couldn't find UnrealBuildTool at '%s'"), *UnrealBuildToolPath);
+		return false;
+	}
+
+	// On Mac we launch UBT with Mono
+	FString ScriptPath = FPaths::ConvertRelativePathToFull(RootDir / TEXT("Engine/Build/BatchFiles/Mac/RunMono.sh"));
+	FString CmdLineParams = FString::Printf(TEXT("\"%s\" \"%s\" %s"), *ScriptPath, *UnrealBuildToolPath, *Arguments);
+
+	// Spawn it
+	int32 ExitCode = 0;
+	return FFeedbackContextMarkup::PipeProcessOutput(Description, TEXT("/bin/sh"), CmdLineParams, Warn, &ExitCode) && ExitCode == 0;
+}
+
 #undef LOCTEXT_NAMESPACE
diff --git a/Engine/Source/Developer/DesktopPlatform/Private/Mac/DesktopPlatformMac.h b/Engine/Source/Developer/DesktopPlatform/Private/Mac/DesktopPlatformMac.h
--- a/Engine/Source/Developer/DesktopPlatform/Private/Mac/DesktopPlatformMac.h
+++ b/Engine/Source/Developer/DesktopPlatform/Private/Mac/DesktopPlatformMac.h
@@ -20,6 +20,8 @@ public:
 	virtual bool VerifyFileAssociations() OVERRIDE;
 	virtual bool UpdateFileAssociations() OVERRIDE;
 
+	virtual bool RunUnrealBuildTool(const FText& Description, const FString& RootDir, const FString& Arguments, FFeedbackContext* Warn) OVERRIDE;
+
 private:
 	bool FileDialogShared(bool bSave, const void* ParentWindowHandle, const FString& DialogTitle, const FString& DefaultPath, const FString& DefaultFile, const FString& FileTypes, uint32 Flags, TArray<FString>& OutFilenames);
 };
diff --git a/Engine/Source/Developer/DesktopPlatform/Private/Windows/DesktopPlatformWindows.cpp b/Engine/Source/Developer/DesktopPlatform/Private/Windows/DesktopPlatformWindows.cpp
--- a/Engine/Source/Developer/DesktopPlatform/Private/Windows/DesktopPlatformWindows.cpp
+++ b/Engine/Source/Developer/DesktopPlatform/Private/Windows/DesktopPlatformWindows.cpp
@@ -1,6 +1,7 @@
 // Copyright 1998-2014 Epic Games, Inc. All Rights Reserved.
 
 #include "DesktopPlatformPrivatePCH.h"
+#include "FeedbackContextMarkup.h"
 
 #include "AllowWindowsPlatformTypes.h"
 	#include <commdlg.h>
@@ -466,6 +467,21 @@ bool FDesktopPlatformWindows::UpdateFileAssociations()
 	return true;
 }
 
+bool FDesktopPlatformWindows::RunUnrealBuildTool(const FText& Description, const FString& RootDir, const FString& Arguments, FFeedbackContext* Warn)
+{
+	// Get the path to UBT
+	FString UnrealBuildToolPath = RootDir / TEXT("Engine/Binaries/DotNET/UnrealBuildTool.exe");
+	if(IFileManager::Get().FileSize(*UnrealBuildToolPath) < 0)
+	{
+		Warn->Logf(ELogVerbosity::Error, TEXT("Couldn't find UnrealBuildTool at '%s'"), *UnrealBuildToolPath);
+		return false;
+	}
+
+	// Spawn UBT
+	int32 ExitCode = 0;
+	return FFeedbackContextMarkup::PipeProcessOutput(Description, UnrealBuildToolPath, Arguments, Warn, &ExitCode) && ExitCode == 0;
+}
+
 void FDesktopPlatformWindows::GetRequiredRegistrySettings(TIndirectArray<FRegistryRootedKey> &RootedKeys)
 {
 	// Get the path to VersionSelector.exe. If we're running from UnrealVersionSelector itself, try to stick with the current configuration.
diff --git a/Engine/Source/Developer/DesktopPlatform/Private/Windows/DesktopPlatformWindows.h b/Engine/Source/Developer/DesktopPlatform/Private/Windows/DesktopPlatformWindows.h
--- a/Engine/Source/Developer/DesktopPlatform/Private/Windows/DesktopPlatformWindows.h
+++ b/Engine/Source/Developer/DesktopPlatform/Private/Windows/DesktopPlatformWindows.h
@@ -23,6 +23,8 @@ public:
 	virtual bool VerifyFileAssociations() OVERRIDE;
 	virtual bool UpdateFileAssociations() OVERRIDE;
 
+	virtual bool RunUnrealBuildTool(const FText& Description, const FString& RootDir, const FString& Arguments, FFeedbackContext* Warn) OVERRIDE;
+
 private:
 	bool FileDialogShared(bool bSave, const void* ParentWindowHandle, const FString& DialogTitle, const FString& DefaultPath, const FString& DefaultFile, const FString& FileTypes, uint32 Flags, TArray<FString>& OutFilenames);
 	void GetRequiredRegistrySettings(TIndirectArray<FRegistryRootedKey> &RootedKeys);
diff --git a/Engine/Source/Developer/DesktopPlatform/Public/IDesktopPlatform.h b/Engine/Source/Developer/DesktopPlatform/Public/IDesktopPlatform.h
--- a/Engine/Source/Developer/DesktopPlatform/Public/IDesktopPlatform.h
+++ b/Engine/Source/Developer/DesktopPlatform/Public/IDesktopPlatform.h
@@ -189,6 +189,14 @@ public:
 	*/
 	virtual bool GetDefaultEngineRootDir(FString &OutRootDir) = 0;
 
+	/**
+	* Checks if the given engine identifier is for an stock engine release.
+	*
+	* @param	Identifier			Engine identifier to check
+	* @return	true if the identifier is for a binary engine release
+	*/
+	virtual bool IsStockEngineRelease(const FString &Identifier) = 0;
+
 	/**
 	* Tests whether an engine installation is a source distribution.
 	*
@@ -241,4 +249,35 @@ public:
 	* @return true if OutIdentifier is set to the project's engine association
 	*/
 	virtual bool GetEngineIdentifierForProject(const FString &ProjectFileName, FString &OutIdentifier) = 0;
+
+	/**
+	* Compiles a game project.
+	*
+	* @param RootDir			Engine root directory for the project to use.
+	* @param ProjectFileName	Filename of the project to update
+	* @param Warn				Feedback context to use for progress updates
+	* @return true if project files were generated successfully.
+	*/
+	virtual bool CompileGameProject(const FString& RootDir, const FString& ProjectFileName, FFeedbackContext* Warn) = 0;
+
+	/**
+	* Generates project files for the given project.
+	*
+	* @param RootDir			Engine root directory for the project to use.
+	* @param ProjectFileName	Filename of the project to update
+	* @param Warn				Feedback context to use for progress updates
+	* @return true if project files were generated successfully.
+	*/
+	virtual bool GenerateProjectFiles(const FString& RootDir, const FString& ProjectFileName, FFeedbackContext* Warn) = 0;
+
+	/**
+	* Runs UnrealBuildTool with the given arguments.
+	*
+	* @param Description		Task description for FFeedbackContext
+	* @param RootDir			Engine root directory for the project to use.
+	* @param Arguments			Parameters for UnrealBuildTool
+	* @param Warn				Feedback context to use for progress updates
+	* @return true if the task completed successfully.
+	*/
+	virtual bool RunUnrealBuildTool(const FText& Description, const FString& RootDir, const FString& Arguments, FFeedbackContext* Warn) = 0;
 };
diff --git a/Engine/Source/Editor/GameProjectGeneration/Private/GameProjectUtils.cpp b/Engine/Source/Editor/GameProjectGeneration/Private/GameProjectUtils.cpp
--- a/Engine/Source/Editor/GameProjectGeneration/Private/GameProjectUtils.cpp
+++ b/Engine/Source/Editor/GameProjectGeneration/Private/GameProjectUtils.cpp
@@ -377,11 +377,11 @@ void GameProjectUtils::CheckForOutOfDateGameProjectFile()
 	if ( !LoadedProjectFilePath.IsEmpty() )
 	{
 		FProjectStatus ProjectStatus;
-		if (IProjectManager::Get().QueryStatusForProject(LoadedProjectFilePath, FDesktopPlatformModule::Get()->GetCurrentEngineIdentifier(), ProjectStatus))
+		if (IProjectManager::Get().QueryStatusForProject(LoadedProjectFilePath, ProjectStatus))
 		{
-			if ( !ProjectStatus.bUpToDate )
+			if ( ProjectStatus.bRequiresUpdate )
 			{
-				const FText UpdateProjectText = LOCTEXT("UpdateProjectFilePrompt", "Update project to open in this version of the editor by default?");
+				const FText UpdateProjectText = LOCTEXT("UpdateProjectFilePrompt", "Project file is saved in an older format. Would you like to update it?");
 				const FText UpdateProjectConfirmText = LOCTEXT("UpdateProjectFileConfirm", "Update");
 				const FText UpdateProjectCancelText = LOCTEXT("UpdateProjectFileCancel", "Not Now");
 
@@ -463,29 +463,17 @@ bool GameProjectUtils::UpdateGameProject(const FString& EngineIdentifier)
 	const FString& ProjectFilename = FPaths::IsProjectFilePathSet() ? FPaths::GetProjectFilePath() : FString();
 	if ( !ProjectFilename.IsEmpty() )
 	{
-		FProjectStatus ProjectStatus;
-		if ( IProjectManager::Get().QueryStatusForProject(ProjectFilename, EngineIdentifier, ProjectStatus) )
+		FText FailReason;
+		bool bWasCheckedOut = false;
+		if ( !UpdateGameProjectFile(ProjectFilename, EngineIdentifier, NULL, bWasCheckedOut, FailReason) )
 		{
-			if ( ProjectStatus.bUpToDate )
-			{
-				// The project was already up to date.
-				UE_LOG(LogGameProjectGeneration, Log, TEXT("%s is already up to date."), *ProjectFilename );
-			}
-			else
-			{
-				FText FailReason;
-				bool bWasCheckedOut = false;
-				if ( !UpdateGameProjectFile(ProjectFilename, EngineIdentifier, NULL, bWasCheckedOut, FailReason) )
-				{
-					// The user chose to update, but the update failed. Notify the user.
-					UE_LOG(LogGameProjectGeneration, Error, TEXT("%s failed to update. %s"), *ProjectFilename, *FailReason.ToString() );
-					return false;
-				}
-
-				// The project was updated successfully.
-				UE_LOG(LogGameProjectGeneration, Log, TEXT("%s was successfully updated."), *ProjectFilename );
-			}
+			// The user chose to update, but the update failed. Notify the user.
+			UE_LOG(LogGameProjectGeneration, Error, TEXT("%s failed to update. %s"), *ProjectFilename, *FailReason.ToString() );
+			return false;
 		}
+
+		// The project was updated successfully.
+		UE_LOG(LogGameProjectGeneration, Log, TEXT("%s was successfully updated."), *ProjectFilename );
 	}
 
 	return true;
@@ -1590,6 +1578,102 @@ bool GameProjectUtils::CalculateSourcePaths(const FString& InPath, FString& OutM
 	return !OutHeaderPath.IsEmpty() && !OutSourcePath.IsEmpty();
 }
 
+bool GameProjectUtils::DuplicateProjectForUpgrade( const FString& InProjectFile, FString &OutNewProjectFile )
+{
+	IPlatformFile &PlatformFile = FPlatformFileManager::Get().GetPlatformFile();
+
+	// Get the directory part of the project name
+	FString OldDirectoryName = FPaths::GetPath(InProjectFile);
+	FPaths::NormalizeDirectoryName(OldDirectoryName);
+	FString NewDirectoryName = OldDirectoryName;
+
+	// Strip off any previous version number from the project name
+	for(int32 LastSpace; NewDirectoryName.FindLastChar(' ', LastSpace); )
+	{
+		const TCHAR *End = *NewDirectoryName + LastSpace + 1;
+		if(*End == '(')
+		{
+			End++;
+			while(FChar::IsDigit(*End) || *End == '.')
+			{
+				End++;
+			}
+			if(*End == ')' && *(End + 1) == 0)
+			{
+				NewDirectoryName = NewDirectoryName.Left(LastSpace).TrimTrailing();
+			}
+		}
+	}
+
+	// Append the new version number
+	NewDirectoryName += FString::Printf(TEXT(" (%s)"), *GEngineVersion.ToString(EVersionComponent::Minor));
+
+	// Find a directory name that doesn't exist
+	FString BaseDirectoryName = NewDirectoryName;
+	for(int32 Idx = 2; IFileManager::Get().DirectoryExists(*NewDirectoryName); Idx++)
+	{
+		NewDirectoryName = FString::Printf(TEXT("%s (%d)"), *BaseDirectoryName, Idx);
+	}
+
+	// Find all the root directory names
+	TArray<FString> RootDirectoryNames;
+	IFileManager::Get().FindFiles(RootDirectoryNames, *(OldDirectoryName / TEXT("*")), false, true);
+
+	// Find all the source directories
+	TArray<FString> SourceDirectories;
+	SourceDirectories.Add(OldDirectoryName);
+	for(int32 Idx = 0; Idx < RootDirectoryNames.Num(); Idx++)
+	{
+		if(RootDirectoryNames[Idx] != TEXT("Binaries") && RootDirectoryNames[Idx] != TEXT("Intermediate") && RootDirectoryNames[Idx] != TEXT("Saved"))
+		{
+			FString SourceDirectory = OldDirectoryName / RootDirectoryNames[Idx];
+			SourceDirectories.Add(SourceDirectory);
+			IFileManager::Get().FindFilesRecursive(SourceDirectories, *SourceDirectory, TEXT("*"), false, true, false);
+		}
+	}
+
+	// Find all the source files
+	TArray<FString> SourceFiles;
+	for(int32 Idx = 0; Idx < SourceDirectories.Num(); Idx++)
+	{
+		TArray<FString> SourceNames;
+		IFileManager::Get().FindFiles(SourceNames, *(SourceDirectories[Idx] / TEXT("*")), true, false);
+
+		for(int32 NameIdx = 0; NameIdx < SourceNames.Num(); NameIdx++)
+		{
+			SourceFiles.Add(SourceDirectories[Idx] / SourceNames[NameIdx]);
+		}
+	}
+
+	// Copy everything
+	bool bCopySucceeded = true;
+	GWarn->BeginSlowTask(LOCTEXT("CreatingCopyOfProject", "Creating copy of project..."), true);
+	for(int32 Idx = 0; Idx < SourceDirectories.Num() && bCopySucceeded; Idx++)
+	{
+		FString TargetDirectory = NewDirectoryName + SourceDirectories[Idx].Mid(OldDirectoryName.Len());
+		bCopySucceeded = PlatformFile.CreateDirectory(*TargetDirectory);
+		GWarn->UpdateProgress(Idx + 1, SourceDirectories.Num() + SourceFiles.Num());
+	}
+	for(int32 Idx = 0; Idx < SourceFiles.Num() && bCopySucceeded; Idx++)
+	{
+		FString TargetFile = NewDirectoryName + SourceFiles[Idx].Mid(OldDirectoryName.Len());
+		bCopySucceeded = PlatformFile.CopyFile(*TargetFile, *SourceFiles[Idx]);
+		GWarn->UpdateProgress(SourceDirectories.Num() + Idx + 1, SourceDirectories.Num() + SourceFiles.Num());
+	}
+	GWarn->EndSlowTask();
+
+	// Wipe the directory if we couldn't update
+	if(!bCopySucceeded)
+	{
+		PlatformFile.DeleteDirectoryRecursively(*NewDirectoryName);
+		return false;
+	}
+
+	// Otherwise fixup the output project filename
+	OutNewProjectFile = NewDirectoryName / FPaths::GetCleanFilename(InProjectFile);
+	return true;
+}
+
 bool GameProjectUtils::ReadTemplateFile(const FString& TemplateFileName, FString& OutFileContents, FText& OutFailReason)
 {
 	const FString FullFileName = FPaths::EngineContentDir() / TEXT("Editor") / TEXT("Templates") / TemplateFileName;
diff --git a/Engine/Source/Editor/GameProjectGeneration/Private/GameProjectUtils.h b/Engine/Source/Editor/GameProjectGeneration/Private/GameProjectUtils.h
--- a/Engine/Source/Editor/GameProjectGeneration/Private/GameProjectUtils.h
+++ b/Engine/Source/Editor/GameProjectGeneration/Private/GameProjectUtils.h
@@ -92,6 +92,9 @@ public:
 	 */
 	static bool CalculateSourcePaths(const FString& InPath, FString& OutModuleName, FString& OutHeaderPath, FString& OutSourcePath, FText* const OutFailReason = nullptr);
 
+	/** Creates a copy of a project directory in order to upgrade it. */
+	static bool DuplicateProjectForUpgrade( const FString& InProjectFile, FString &OutNewProjectFile );
+
 private:
 	/** Generates a new project without using a template project */
 	static bool GenerateProjectFromScratch(const FString& NewProjectFile, bool bShouldGenerateCode, bool bCopyStarterContent, FText& OutFailReason);
diff --git a/Engine/Source/Editor/GameProjectGeneration/Private/SProjectBrowser.cpp b/Engine/Source/Editor/GameProjectGeneration/Private/SProjectBrowser.cpp
--- a/Engine/Source/Editor/GameProjectGeneration/Private/SProjectBrowser.cpp
+++ b/Engine/Source/Editor/GameProjectGeneration/Private/SProjectBrowser.cpp
@@ -4,6 +4,7 @@
 #include "GameProjectGenerationPrivatePCH.h"
 #include "MainFrame.h"
 #include "DesktopPlatformModule.h"
+#include "SVerbChoiceDialog.h"
 
 #define LOCTEXT_NAMESPACE "ProjectBrowser"
 
@@ -15,16 +16,16 @@ struct FProjectItem
 {
 	FText Name;
 	FText Description;
-	FEngineVersion EngineVersion;
+	FString EngineIdentifier;
 	bool bUpToDate;
 	FString ProjectFile;
 	TSharedPtr<FSlateBrush> ProjectThumbnail;
 	bool bIsNewProjectItem;
 
-	FProjectItem(const FText& InName, const FText& InDescription, const FEngineVersion& InEngineVersion, bool InUpToDate, const TSharedPtr<FSlateBrush>& InProjectThumbnail, const FString& InProjectFile, bool InIsNewProjectItem)
+	FProjectItem(const FText& InName, const FText& InDescription, const FString& InEngineIdentifier, bool InUpToDate, const TSharedPtr<FSlateBrush>& InProjectThumbnail, const FString& InProjectFile, bool InIsNewProjectItem)
 		: Name(InName)
 		, Description(InDescription)
-		, EngineVersion(InEngineVersion)
+		, EngineIdentifier(InEngineIdentifier)
 		, bUpToDate(InUpToDate)
 		, ProjectFile(InProjectFile)
 		, ProjectThumbnail(InProjectThumbnail)
@@ -34,7 +35,24 @@ struct FProjectItem
 	/** Check if this project is up to date */
 	bool IsUpToDate() const
 	{
-		return bUpToDate || EngineVersion.IsEmpty();
+		return bUpToDate;
+	}
+
+	/** Gets the engine label for this project */
+	FString GetEngineLabel() const
+	{
+		if(bUpToDate)
+		{
+			return FString();
+		}
+		else if(FDesktopPlatformModule::Get()->IsStockEngineRelease(EngineIdentifier))
+		{
+			return EngineIdentifier;
+		}
+		else
+		{
+			return FString(TEXT("?"));
+		}
 	}
 };
 
@@ -365,7 +383,7 @@ TSharedRef<ITableRow> SProjectBrowser::MakeProjectViewWidget(TSharedPtr<FProject
 					.Padding(10)
 					[
 						SNew(STextBlock)
-						.Text(ProjectItem->EngineVersion.ToString(EVersionComponent::Minor))
+						.Text(ProjectItem->GetEngineLabel())
 						.TextStyle(FEditorStyle::Get(), "ProjectBrowser.VersionOverlayText")
 						.ColorAndOpacity(FLinearColor::White.CopyWithNewOpacity(0.5f))
 						.Visibility(ProjectItem->IsUpToDate() ? EVisibility::Collapsed : EVisibility::Visible)
@@ -408,10 +426,26 @@ TSharedRef<SToolTip> SProjectBrowser::MakeProjectToolTip( TSharedPtr<FProjectIte
 
 	if (!ProjectItem->IsUpToDate())
 	{
-		FFormatOrderedArguments Args;
-		Args.Add(FText::FromString(ProjectItem->EngineVersion.ToString(EVersionComponent::Minor)));
-		AddToToolTipInfoBox(InfoBox, LOCTEXT("ProjectVersion", "Version"),
-			FText::Format(LOCTEXT("ProjectTileOutOfDate", "This project was created with engine version {0}. It may require an update to work with this build."), Args));
+		FText Description;
+		if(FDesktopPlatformModule::Get()->IsStockEngineRelease(ProjectItem->EngineIdentifier))
+		{
+			Description = FText::FromString(ProjectItem->EngineIdentifier);
+		}
+		else
+		{
+			FString RootDir;
+			if(FDesktopPlatformModule::Get()->GetEngineRootDirFromIdentifier(ProjectItem->EngineIdentifier, RootDir))
+			{
+				FString PlatformRootDir = RootDir;
+				FPaths::MakePlatformFilename(PlatformRootDir);
+				Description = FText::FromString(PlatformRootDir);
+			}
+			else
+			{
+				Description = LOCTEXT("UnknownEngineVersion", "Unknown engine version");
+			}
+		}
+		AddToToolTipInfoBox(InfoBox, LOCTEXT("EngineVersion", "Engine"), Description);
 	}
 
 	TSharedRef<SToolTip> Tooltip = SNew(SToolTip)
@@ -655,7 +689,7 @@ FReply SProjectBrowser::FindProjects()
 		{
 			const bool bPromptIfSavedWithNewerVersionOfEngine = false;
 			FProjectStatus ProjectStatus;
-			if (IProjectManager::Get().QueryStatusForProject(ProjectFilename, EngineIdentifier, ProjectStatus))
+			if (IProjectManager::Get().QueryStatusForProject(ProjectFilename, ProjectStatus))
 			{
 				// @todo localized project name
 				const FText ProjectName = FText::FromString(ProjectStatus.Name);
@@ -690,8 +724,11 @@ FReply SProjectBrowser::FindProjects()
 					ProjectCategory = DetectedCategory;
 				}
 
+				FString ProjectEngineIdentifier;
+				bool bIsUpToDate = FDesktopPlatformModule::Get()->GetEngineIdentifierForProject(ProjectFilename, ProjectEngineIdentifier) && ProjectEngineIdentifier == EngineIdentifier;
+
 				const bool bIsNewProjectItem = false;
-				TSharedRef<FProjectItem> NewProjectItem = MakeShareable( new FProjectItem(ProjectName, ProjectDescription, ProjectStatus.EngineVersion, ProjectStatus.bUpToDate, DynamicBrush, ProjectFilename, bIsNewProjectItem ) );
+				TSharedRef<FProjectItem> NewProjectItem = MakeShareable( new FProjectItem(ProjectName, ProjectDescription, ProjectEngineIdentifier, bIsUpToDate, DynamicBrush, ProjectFilename, bIsNewProjectItem ) );
 				AddProjectToCategory(NewProjectItem, ProjectCategory);
 			}
 		}
@@ -814,15 +851,119 @@ FReply SProjectBrowser::OnKeyDown( const FGeometry& MyGeometry, const FKeyboardE
 	return FReply::Unhandled();
 }
 
-
-bool SProjectBrowser::OpenProject( const FString& ProjectFile )
+bool SProjectBrowser::OpenProject( const FString& InProjectFile )
 {
 	FText FailReason;
+	FString ProjectFile = InProjectFile;
+
+	// Get the identifier for the project
+	FString ProjectIdentifier;
+	FDesktopPlatformModule::Get()->GetEngineIdentifierForProject(ProjectFile, ProjectIdentifier);
+	
+	// Get the identifier for the current engine
+	FString CurrentIdentifier = FDesktopPlatformModule::Get()->GetCurrentEngineIdentifier();
+	if(ProjectIdentifier != CurrentIdentifier)
+	{
+		// Get the current project status
+		FProjectStatus ProjectStatus;
+		if(!IProjectManager::Get().QueryStatusForProject(ProjectFile, ProjectStatus))
+		{
+			FMessageDialog::Open(EAppMsgType::Ok, LOCTEXT("CouldNotReadProjectStatus", "Unable to read project status."));
+			return false;
+		}
 
+		// Ask whether to copy the project or upgrade in-place
+		TArray<FText> Buttons;
+		int32 OpenCopyButton = Buttons.Add(LOCTEXT("ProjectUpgradeCopy", "Create a copy"));
+		Buttons.Add(LOCTEXT("ProjectUpgradeInPlace", "Convert in-place"));
+		int32 CancelButton = Buttons.Add(LOCTEXT("ProjectUpgradeCancel", "Cancel"));
+
+		// Show the dialog
+		int32 Selection = SVerbChoiceDialog::ShowModal(LOCTEXT("ProjectUpgradeTitle", "Upgrade Project"), LOCTEXT("ProjectUpgradePrompt", "This project was made with a different version of the editor. Converting it to this version of the editor may be irreversible.\n\nWould you like to create a copy before opening?"), Buttons);
+		if(Selection == OpenCopyButton)
+		{
+			FString NewProjectFile;
+			if(!GameProjectUtils::DuplicateProjectForUpgrade(ProjectFile, NewProjectFile))
+			{
+				FMessageDialog::Open( EAppMsgType::Ok, LOCTEXT("CopyFailed", "Couldn't copy project. Check you have sufficient hard drive space and write access to the project folder.") );
+				return false;
+			}
+			ProjectFile = NewProjectFile;
+		}
+		else if(Selection == CancelButton)
+		{
+			return false;
+		}
+
+		// If it's a code-based project, generate project files and open visual studio after an upgrade
+		if(ProjectStatus.bCodeBasedProject)
+		{
+			// Try to generate project files
+			FStringOutputDevice OutputLog;
+			OutputLog.SetAutoEmitLineTerminator(true);
+			GLog->AddOutputDevice(&OutputLog);
+			bool bHaveProjectFiles = FDesktopPlatformModule::Get()->GenerateProjectFiles(FPaths::RootDir(), ProjectFile, GWarn);
+			GLog->RemoveOutputDevice(&OutputLog);
+
+			// Display any errors
+			if(!bHaveProjectFiles)
+			{
+				FFormatNamedArguments Args;
+				Args.Add( TEXT("LogOutput"), FText::FromString(OutputLog) );
+				FMessageDialog::Open(EAppMsgType::Ok, FText::Format(LOCTEXT("CouldNotGenerateProjectFiles", "Project files could not be generated. Log output:\n\n{LogOutput}"), Args));
+				return false;
+			}
+
+			// Try to compile the project
+			OutputLog.Empty();
+			GLog->AddOutputDevice(&OutputLog);
+			bool bCompileSucceeded = FDesktopPlatformModule::Get()->CompileGameProject(FPaths::RootDir(), ProjectFile, GWarn);
+			GLog->RemoveOutputDevice(&OutputLog);
+
+			// Try to compile the modules
+			if(!bCompileSucceeded)
+			{
+				FText DevEnvName = FSourceCodeNavigation::GetSuggestedSourceCodeIDE( true );
+
+				TArray<FText> CompileFailedButtons;
+				int32 OpenIDEButton = CompileFailedButtons.Add(FText::Format(LOCTEXT("CompileFailedOpenIDE", "Open with {0}"), DevEnvName));
+				int32 ViewLogButton = CompileFailedButtons.Add(LOCTEXT("CompileFailedViewLog", "View build log"));
+				CompileFailedButtons.Add(LOCTEXT("CompileFailedCancel", "Cancel"));
+
+				int32 CompileFailedChoice = SVerbChoiceDialog::ShowModal(LOCTEXT("ProjectUpgradeTitle", "Project Conversion Failed"), FText::Format(LOCTEXT("ProjectUpgradeCompileFailed", "The project failed to compile with this version of the engine. Would you like to open the project in {0}?"), DevEnvName), CompileFailedButtons);
+				if(CompileFailedChoice == ViewLogButton)
+				{
+					CompileFailedButtons.RemoveAt(ViewLogButton);
+					CompileFailedChoice = SVerbChoiceDialog::ShowModal(LOCTEXT("ProjectUpgradeTitle", "Project Conversion Failed"), FText::Format(LOCTEXT("ProjectUpgradeCompileFailed", "The project failed to compile with this version of the engine. Build output is as follows:\n\n{0}"), FText::FromString(OutputLog)), CompileFailedButtons);
+				}
+
+				if(CompileFailedChoice == OpenIDEButton && !GameProjectUtils::OpenCodeIDE(ProjectFile, FailReason))
+				{
+					FMessageDialog::Open(EAppMsgType::Ok, FailReason);
+				}
+
+				return false;
+			}
+		}
+
+		// Update the game project to the latest version. This will prompt to check-out as necessary. We don't need to write the engine identifier yet, because it won't use the right .uprojectdirs logic.
+		if(!GameProjectUtils::UpdateGameProject(TEXT("")))
+		{
+			return false;
+		}
+
+		// Write the engine association
+		if(!FDesktopPlatformModule::Get()->SetEngineIdentifierForProject(ProjectFile, CurrentIdentifier))
+		{
+			FMessageDialog::Open(EAppMsgType::Ok, LOCTEXT("ProjectUpgradeFailure", "Project file could not be updated to latest version. Check that the file is writable."));
+			return false;
+		}
+	}
+
+	// Open the project
 	if (!GameProjectUtils::OpenProject(ProjectFile, FailReason))
 	{
 		FMessageDialog::Open(EAppMsgType::Ok, FailReason);
-
 		return false;
 	}
 
diff --git a/Engine/Source/Editor/GameProjectGeneration/Private/SVerbChoiceDialog.cpp b/Engine/Source/Editor/GameProjectGeneration/Private/SVerbChoiceDialog.cpp
new file mode 100644
--- /dev/null
+++ b/Engine/Source/Editor/GameProjectGeneration/Private/SVerbChoiceDialog.cpp
@@ -0,0 +1,153 @@
+// Copyright 1998-2014 Epic Games, Inc. All Rights Reserved.
+
+#include "GameProjectGenerationPrivatePCH.h"
+#include "SVerbChoiceDialog.h"
+
+int32 SVerbChoiceDialog::ShowModal( const FText& InTitle, const FText& InMessage, const TArray<FText>& InButtons )
+{
+	TSharedRef<SWindow> ModalWindow = SNew(SWindow)
+		.Title( InTitle )
+		.SizingRule( ESizingRule::Autosized )
+		.AutoCenter(EAutoCenter::PreferredWorkArea)
+		.SupportsMinimize(false) 
+		.SupportsMaximize(false);
+
+	TSharedRef<SVerbChoiceDialog> MessageBox = SNew(SVerbChoiceDialog)
+		.ParentWindow(ModalWindow)
+		.Message( InMessage )
+		.Buttons( InButtons )
+		.WrapMessageAt(512.0f);
+
+	ModalWindow->SetContent( MessageBox );
+
+	GEditor->EditorAddModalWindow(ModalWindow);
+
+	return MessageBox->Response;
+}
+
+BEGIN_SLATE_FUNCTION_BUILD_OPTIMIZATION
+void SVerbChoiceDialog::Construct( const FArguments& InArgs )
+{
+	ParentWindow = InArgs._ParentWindow.Get();
+	ParentWindow->SetWidgetToFocusOnActivate(SharedThis(this));
+	Response = EAppReturnType::Cancel;
+
+	FSlateFontInfo MessageFont( FEditorStyle::GetFontStyle("StandardDialog.LargeFont"));
+	Message = InArgs._Message;
+	Buttons = InArgs._Buttons;
+
+	TSharedPtr<SUniformGridPanel> ButtonBox;
+
+	this->ChildSlot
+		[	
+			SNew(SBorder)
+				.BorderImage(FEditorStyle::GetBrush("ToolPanel.GroupBorder"))
+				[
+					SNew(SVerticalBox)
+
+					+ SVerticalBox::Slot()
+						.HAlign(HAlign_Fill)
+						.VAlign(VAlign_Fill)
+						.FillHeight(1.0f)
+						.MaxHeight(550)
+						.Padding(12.0f)
+						[
+							SNew(SScrollBox)
+
+							+ SScrollBox::Slot()
+								[
+									SNew(STextBlock)
+										.Text(Message)
+										.Font(MessageFont)
+										.WrapTextAt(InArgs._WrapMessageAt)
+								]
+						]
+
+					+SVerticalBox::Slot()
+						.AutoHeight()
+						.Padding(0.0f)
+						[
+							SNew(SHorizontalBox)
+
+							+ SHorizontalBox::Slot()
+								.FillWidth(1.0f)
+								.HAlign(HAlign_Left)
+								.VAlign(VAlign_Bottom)
+								.Padding(12.0f)
+								[
+									SNew(SHyperlink)
+										.OnNavigate(this, &SVerbChoiceDialog::HandleCopyMessageHyperlinkNavigate)
+										.Text( NSLOCTEXT("SVerbChoiceDialog", "CopyMessageHyperlink", "Copy Message") )
+										.ToolTipText( NSLOCTEXT("SVerbChoiceDialog", "CopyMessageTooltip", "Copy the text in this message to the clipboard (CTRL+C)") )
+								]
+
+							+ SHorizontalBox::Slot()
+								.AutoWidth()
+								.HAlign(HAlign_Right)
+								.VAlign(VAlign_Bottom)
+								.Padding(2.f)
+								[
+									SAssignNew( ButtonBox, SUniformGridPanel )
+										.SlotPadding(FEditorStyle::GetMargin("StandardDialog.SlotPadding"))
+										.MinDesiredSlotWidth(FEditorStyle::GetFloat("StandardDialog.MinDesiredSlotWidth"))
+										.MinDesiredSlotHeight(FEditorStyle::GetFloat("StandardDialog.MinDesiredSlotHeight"))
+								]
+						]
+				]
+		];
+
+	for(int32 Idx = 0; Idx < Buttons.Get().Num(); Idx++)
+	{
+		ButtonBox->AddSlot(Idx, 0)
+			[
+				SNew( SButton )
+				.Text( Buttons.Get()[Idx] )
+				.OnClicked( this, &SVerbChoiceDialog::HandleButtonClicked, Idx )
+				.ContentPadding(FEditorStyle::GetMargin("StandardDialog.ContentPadding"))
+				.HAlign(HAlign_Center)
+			];
+	}
+}
+END_SLATE_FUNCTION_BUILD_OPTIMIZATION
+
+FReply SVerbChoiceDialog::OnKeyDown( const FGeometry& MyGeometry, const FKeyboardEvent& InKeyboardEvent )
+{
+	//see if we pressed the Enter or Spacebar keys
+	if( InKeyboardEvent.GetKey() == EKeys::Escape )
+	{
+		return HandleButtonClicked(EAppReturnType::Cancel);
+	}
+
+	if (InKeyboardEvent.GetKey() == EKeys::C && InKeyboardEvent.IsControlDown())
+	{
+		CopyMessageToClipboard();
+
+		return FReply::Handled();
+	}
+
+	//if it was some other button, ignore it
+	return FReply::Unhandled();
+}
+
+bool SVerbChoiceDialog::SupportsKeyboardFocus() const
+{
+	return true;
+}
+
+void SVerbChoiceDialog::CopyMessageToClipboard( )
+{
+	FPlatformMisc::ClipboardCopy( *Message.Get().ToString() );
+}
+
+FReply SVerbChoiceDialog::HandleButtonClicked( int32 InResponse )
+{
+	Response = InResponse;
+	ParentWindow->RequestDestroyWindow();
+
+	return FReply::Handled();
+}
+
+void SVerbChoiceDialog::HandleCopyMessageHyperlinkNavigate( )
+{
+	CopyMessageToClipboard();
+}
diff --git a/Engine/Source/Editor/GameProjectGeneration/Private/SVerbChoiceDialog.h b/Engine/Source/Editor/GameProjectGeneration/Private/SVerbChoiceDialog.h
new file mode 100644
--- /dev/null
+++ b/Engine/Source/Editor/GameProjectGeneration/Private/SVerbChoiceDialog.h
@@ -0,0 +1,42 @@
+// Copyright 1998-2014 Epic Games, Inc. All Rights Reserved.
+
+#pragma once
+
+class SVerbChoiceDialog : public SCompoundWidget
+{
+public:
+	SLATE_BEGIN_ARGS( SVerbChoiceDialog )	{}
+		SLATE_ATTRIBUTE(TSharedPtr<SWindow>, ParentWindow)
+		SLATE_ATTRIBUTE(FText, Message)	
+		SLATE_ATTRIBUTE(TArray<FText>, Buttons)
+		SLATE_ATTRIBUTE(float, WrapMessageAt)
+	SLATE_END_ARGS()
+
+	/** Displays the modal dialog box */
+	static int32 ShowModal(const FText& InTitle, const FText& InText, const TArray<FText>& InButtons);
+
+	/** Constructs the dialog */
+	void Construct( const FArguments& InArgs );
+
+	/** Override behavior when a key is pressed */
+	virtual	FReply OnKeyDown( const FGeometry& MyGeometry, const FKeyboardEvent& InKeyboardEvent ) OVERRIDE;
+
+	/** Override the base method to allow for keyboard focus */
+	virtual bool SupportsKeyboardFocus() const OVERRIDE;
+
+protected:
+	/** Copies the message text to the clipboard. */
+	void CopyMessageToClipboard( );
+
+private:
+	/** Handles clicking a message box button. */
+	FReply HandleButtonClicked( int32 InResponse );
+
+	/** Handles clicking the 'Copy Message' hyper link. */
+	void HandleCopyMessageHyperlinkNavigate( );
+
+	int32 Response;
+	TSharedPtr<SWindow> ParentWindow;
+	TAttribute<FText> Message;
+	TAttribute< TArray<FText> > Buttons;
+};
diff --git a/Engine/Source/Editor/UnrealEd/Private/UnrealEdMisc.cpp b/Engine/Source/Editor/UnrealEd/Private/UnrealEdMisc.cpp
--- a/Engine/Source/Editor/UnrealEd/Private/UnrealEdMisc.cpp
+++ b/Engine/Source/Editor/UnrealEd/Private/UnrealEdMisc.cpp
@@ -329,7 +329,7 @@ void FUnrealEdMisc::InitEngineAnalytics()
 			const FString& LoadedProjectFilePath = FPaths::GetProjectFilePath();
 			FProjectStatus ProjectStatus;
 
-			if (IProjectManager::Get().QueryStatusForProject(LoadedProjectFilePath, FDesktopPlatformModule::Get()->GetCurrentEngineIdentifier(), ProjectStatus))
+			if (IProjectManager::Get().QueryStatusForProject(LoadedProjectFilePath, ProjectStatus))
 			{
 				if ( ProjectStatus.bSignedSampleProject )
 				{
diff --git a/Engine/Source/Runtime/Projects/Private/PluginManagerShared.cpp b/Engine/Source/Runtime/Projects/Private/PluginManagerShared.cpp
--- a/Engine/Source/Runtime/Projects/Private/PluginManagerShared.cpp
+++ b/Engine/Source/Runtime/Projects/Private/PluginManagerShared.cpp
@@ -227,21 +227,10 @@ FString FProjectOrPlugin::SerializeToJSON( ) const
 	return JSONOutput;
 }
 
-bool FProjectOrPlugin::IsUpToDate( const FString &EngineIdentifier ) const
+bool FProjectOrPlugin::RequiresUpdate( ) const
 {
 	const FProjectOrPluginInfo& ProjectOrPluginInfo = GetProjectOrPluginInfo();
-
-	if (ProjectOrPluginInfo.FileVersion < VER_LATEST_PROJECT_FILE)
-	{
-		return false;
-	}
-
-	if (ProjectOrPluginInfo.EngineAssociation != EngineIdentifier)
-	{
-		return false;
-	}
-
-	return true;
+	return (ProjectOrPluginInfo.FileVersion < VER_LATEST_PROJECT_FILE);
 }
 
 bool FProjectOrPlugin::AreModulesUpToDate( ) const
diff --git a/Engine/Source/Runtime/Projects/Private/PluginManagerShared.h b/Engine/Source/Runtime/Projects/Private/PluginManagerShared.h
--- a/Engine/Source/Runtime/Projects/Private/PluginManagerShared.h
+++ b/Engine/Source/Runtime/Projects/Private/PluginManagerShared.h
@@ -222,11 +222,11 @@ public:
 	FString SerializeToJSON( ) const;
 
 	/**
-	 * Checks whether this project file is up to date with the engine version.
+	 * Checks whether this project file needs updating to a newer project version.
 	 *
-	 * @return true if the project file is up to date, false otherwise.
+	 * @return true if the project needs updating, false otherwise.
 	 */
-	bool IsUpToDate( const FString &EngineIdentifier ) const;
+	bool RequiresUpdate( ) const;
 
 	/**
 	 * Checks whether the modules for a project or plugin need to be compiled
diff --git a/Engine/Source/Runtime/Projects/Private/ProjectManager.cpp b/Engine/Source/Runtime/Projects/Private/ProjectManager.cpp
--- a/Engine/Source/Runtime/Projects/Private/ProjectManager.cpp
+++ b/Engine/Source/Runtime/Projects/Private/ProjectManager.cpp
@@ -17,32 +17,6 @@ FProject::FProject( const FProjectInfo& InitProjectInfo )
 {
 }
 
-bool FProject::VerifyProjectCanBeLoaded( bool bPromptIfSavedWithNewerVersionOfEngine, FText& OutFailureReason )
-{
-	// Verify we can load this project
-	const bool bPackageFileVersionIsValid = (ProjectInfo.PackageFileUE4Version == INDEX_NONE) || (ProjectInfo.PackageFileUE4Version <= GPackageFileUE4Version);
-	const bool bPackageFileLicenseeVersionIsValid = (ProjectInfo.PackageFileLicenseeUE4Version == INDEX_NONE) || (ProjectInfo.PackageFileLicenseeUE4Version <= GPackageFileLicenseeUE4Version);
-	if (!bPackageFileVersionIsValid || !bPackageFileLicenseeVersionIsValid)
-	{
-		EAppReturnType::Type Response = EAppReturnType::No;
-
-		if ( bPromptIfSavedWithNewerVersionOfEngine )
-		{
-			// The project was made with a newer version of the engine. We do not support forward compatibility in content.
-			const FText WarningMessage = FText::Format( LOCTEXT("NewerEngineVersionWarning", "'{0}' was saved with a newer version of the engine. Some content may not load properly. Would you like to attempt to load this project anyway?"), FText::FromString(*ProjectInfo.LoadedFromFile));
-			Response = FMessageDialog::Open(EAppMsgType::YesNo, WarningMessage);
-		}
-
-		if (Response == EAppReturnType::No)
-		{
-			OutFailureReason = FText::Format( LOCTEXT("NewerEngineVersion", "'{0}' was saved with a newer version of the engine."), FText::FromString(*ProjectInfo.LoadedFromFile));
-			return false;
-		}
-	}
-
-	return true;
-}
-
 bool FProject::IsSignedSampleProject(const FString& FilePath) const
 {
 	return ProjectInfo.EpicSampleNameHash == GetTypeHash(FPaths::GetCleanFilename(FilePath));
@@ -81,14 +55,10 @@ bool FProjectManager::LoadProjectFile( const FString& InProjectFile )
 	TSharedRef<FProject> NewProject = MakeShareable( new FProject() );
 	if ( NewProject->LoadFromFile(InProjectFile, FailureReason) )
 	{
-		const bool bPromptIfSavedWithNewerVersionOfEngine = true;
-		if ( NewProject->VerifyProjectCanBeLoaded(bPromptIfSavedWithNewerVersionOfEngine, FailureReason) )
-		{
-			// Load successful. Set the loaded project file pointer.
-			CurrentlyLoadedProject = NewProject;
+		// Load successful. Set the loaded project file pointer.
+		CurrentlyLoadedProject = NewProject;
 
-			return true;
-		}
+		return true;
 	}
 	
 #if PLATFORM_IOS
@@ -270,12 +240,6 @@ bool FProjectManager::SignSampleProject(const FString& FilePath, const FString&
 		return false;
 	}
 
-	const bool bPromptIfSavedWithNewerVersionOfEngine = false;
-	if ( !NewProject->VerifyProjectCanBeLoaded(bPromptIfSavedWithNewerVersionOfEngine, OutFailReason) )
-	{
-		return false;
-	}
-
 	NewProject->SignSampleProject(FilePath, Category);
 
 	const FString& FileContents = NewProject->SerializeToJSON();
@@ -290,7 +254,7 @@ bool FProjectManager::SignSampleProject(const FString& FilePath, const FString&
 	}
 }
 
-bool FProjectManager::QueryStatusForProject(const FString& FilePath, const FString& EngineIdentifier, FProjectStatus& OutProjectStatus) const
+bool FProjectManager::QueryStatusForProject(const FString& FilePath, FProjectStatus& OutProjectStatus) const
 {
 	TSharedRef<FProject> NewProject = MakeShareable( new FProject() );
 	FText FailReason;
@@ -299,19 +263,13 @@ bool FProjectManager::QueryStatusForProject(const FString& FilePath, const FStri
 		return false;
 	}
 
-	const bool bPromptIfSavedWithNewerVersionOfEngine = false;
-	if ( !NewProject->VerifyProjectCanBeLoaded(bPromptIfSavedWithNewerVersionOfEngine, FailReason) )
-	{
-		return false;
-	}
-
 	const FProjectInfo& ProjectInfo = NewProject->GetProjectInfo();
 	OutProjectStatus.Name = ProjectInfo.Name;
 	OutProjectStatus.Description = ProjectInfo.Description;
 	OutProjectStatus.Category = ProjectInfo.Category;
-	OutProjectStatus.EngineVersion = ProjectInfo.EngineVersion;
+	OutProjectStatus.bCodeBasedProject = ProjectInfo.Modules.Num() > 0;
 	OutProjectStatus.bSignedSampleProject = NewProject->IsSignedSampleProject(FilePath);
-	OutProjectStatus.bUpToDate = NewProject->IsUpToDate(EngineIdentifier);
+	OutProjectStatus.bRequiresUpdate = NewProject->RequiresUpdate();
 
 	return true;
 }
diff --git a/Engine/Source/Runtime/Projects/Private/ProjectManager.h b/Engine/Source/Runtime/Projects/Private/ProjectManager.h
--- a/Engine/Source/Runtime/Projects/Private/ProjectManager.h
+++ b/Engine/Source/Runtime/Projects/Private/ProjectManager.h
@@ -38,8 +38,6 @@ public:
 	FProject();
 	FProject( const FProjectInfo& InitProjectInfo );
 
-	bool VerifyProjectCanBeLoaded( bool bPromptIfSavedWithNewerVersionOfEngine, FText& OutFailureReason );
-
 	/**
 	 * Gets whether the project is a sample project signed with the given filename.
 	 *
@@ -93,7 +91,7 @@ public:
 	virtual bool DuplicateProjectFile(const FString& SourceProjectFilename, const FString& NewProjectFilename, const FString& EngineIdentifier, FText& OutFailReason) OVERRIDE;
 	virtual bool UpdateLoadedProjectFileToCurrent(const TArray<FString>* StartupModuleNames, const FString& EngineIdentifier, FText& OutFailReason) OVERRIDE;
 	virtual bool SignSampleProject(const FString& FilePath, const FString& Category, FText& OutFailReason) OVERRIDE;
-	virtual bool QueryStatusForProject(const FString& FilePath, const FString& EngineIdentifier, FProjectStatus& OutProjectStatus) const OVERRIDE;
+	virtual bool QueryStatusForProject(const FString& FilePath, FProjectStatus& OutProjectStatus) const OVERRIDE;
 
 private:
 	/** The project that is currently loaded in the editor */
diff --git a/Engine/Source/Runtime/Projects/Public/Interfaces/IProjectManager.h b/Engine/Source/Runtime/Projects/Public/Interfaces/IProjectManager.h
--- a/Engine/Source/Runtime/Projects/Public/Interfaces/IProjectManager.h
+++ b/Engine/Source/Runtime/Projects/Public/Interfaces/IProjectManager.h
@@ -22,18 +22,19 @@ public:
 	/** The UI category of this project */
 	FString Category;
 
-	/** The version of the engine that this project is made for */
-	FEngineVersion EngineVersion;
-
 	/** True if this project is a sample provided by epic */
 	bool bSignedSampleProject;
 
-	/** True if this project is up to date with the current engine version */
-	bool bUpToDate;
+	/** True if the project is code-based */
+	bool bCodeBasedProject;
+
+	/** True if this project needs to be updated */
+	bool bRequiresUpdate;
 
 	FProjectStatus()
 		: bSignedSampleProject(false)
-		, bUpToDate(false)
+		, bCodeBasedProject(false)
+		, bRequiresUpdate(false)
 	{}
 };
 
@@ -151,7 +152,7 @@ public:
 	 *
 	 * @return	 true if the file was successfully open and read
 	 */
-	virtual bool QueryStatusForProject(const FString& FilePath, const FString& EngineIdentifier, FProjectStatus& OutProjectStatus) const = 0;
+	virtual bool QueryStatusForProject(const FString& FilePath, FProjectStatus& OutProjectStatus) const = 0;
 
 	/** Helper functions to reduce the syntax complexity of commonly used functions */
 	static const FString& GetProjectFileExtension() { return Get().NonStaticGetProjectFileExtension(); }
