Improve project discovery in the project browser.

- All projects referenced by .uprojectdirs are included, rather than using a separate hard-coded list of paths.
- Samples installed by the launcher are included.

diff --git a/Engine/Source/Developer/DesktopPlatform/Private/DesktopPlatformBase.cpp b/Engine/Source/Developer/DesktopPlatform/Private/DesktopPlatformBase.cpp
--- a/Engine/Source/Developer/DesktopPlatform/Private/DesktopPlatformBase.cpp
+++ b/Engine/Source/Developer/DesktopPlatform/Private/DesktopPlatformBase.cpp
@@ -40,6 +40,23 @@ void FDesktopPlatformBase::EnumerateLauncherEngineInstallations(TMap<FString, FS
 	}
 }
 
+void FDesktopPlatformBase::EnumerateLauncherSampleInstallations(TArray<FString> &OutInstallations)
+{
+	TMap<FString, FString> AllInstallations;
+	if(ReadLauncherInstallationList(AllInstallations))
+	{
+		// We've got a list of launcher installations. Filter it by non-engine installations.
+		for(TMap<FString, FString>::TConstIterator Iter(AllInstallations); Iter; ++Iter)
+		{
+			FString AppName = Iter.Key();
+			if(!AppName.StartsWith(TEXT("UE_"), ESearchCase::CaseSensitive))
+			{
+				OutInstallations.Add(Iter.Value());
+			}
+		}
+	}
+}
+
 bool FDesktopPlatformBase::GetEngineRootDirFromIdentifier(const FString &Identifier, FString &OutRootDir)
 {
 	// Get all the installations
diff --git a/Engine/Source/Developer/DesktopPlatform/Private/DesktopPlatformBase.h b/Engine/Source/Developer/DesktopPlatform/Private/DesktopPlatformBase.h
--- a/Engine/Source/Developer/DesktopPlatform/Private/DesktopPlatformBase.h
+++ b/Engine/Source/Developer/DesktopPlatform/Private/DesktopPlatformBase.h
@@ -13,6 +13,7 @@ public:
 	virtual FString GetCurrentEngineIdentifier() OVERRIDE;
 
 	virtual void EnumerateLauncherEngineInstallations(TMap<FString, FString> &OutInstallations) OVERRIDE;
+	virtual void EnumerateLauncherSampleInstallations(TArray<FString> &OutInstallations) OVERRIDE;
 	virtual bool GetEngineRootDirFromIdentifier(const FString &Identifier, FString &OutRootDir) OVERRIDE;
 	virtual bool GetEngineIdentifierFromRootDir(const FString &RootDir, FString &OutIdentifier) OVERRIDE;
 
diff --git a/Engine/Source/Developer/DesktopPlatform/Public/IDesktopPlatform.h b/Engine/Source/Developer/DesktopPlatform/Public/IDesktopPlatform.h
--- a/Engine/Source/Developer/DesktopPlatform/Public/IDesktopPlatform.h
+++ b/Engine/Source/Developer/DesktopPlatform/Public/IDesktopPlatform.h
@@ -136,7 +136,7 @@ public:
 	/**
 	* Enumerates all the registered engine installations.
 	*
-	* @param	OutInstallations	Array which is filled in with identifier/root-directory pairs for all known installations. Identifiers are typically
+	* @param	OutInstallations	Map of identifier/root-directory pairs for all known installations. Identifiers are typically
 	*								version strings for canonical UE4 releases or GUID strings for GitHub releases.
 	*/
 	virtual void EnumerateEngineInstallations(TMap<FString, FString> &OutInstallations) = 0;
@@ -144,10 +144,17 @@ public:
 	/**
 	* Enumerates all the registered binary engine installations.
 	*
-	* @param	OutInstallations	Array which is filled in with identifier/root-directory pairs for all known binary installations.
+	* @param	OutInstallations	Map of identifier/root-directory pairs for all known binary installations.
 	*/
 	virtual void EnumerateLauncherEngineInstallations(TMap<FString, FString> &OutInstallations) = 0;
 
+	/**
+	* Enumerates all the samples installed by the launcher.
+	*
+	* @param	OutInstallations	Array of sample installation paths.
+	*/
+	virtual void EnumerateLauncherSampleInstallations(TArray<FString> &OutInstallations) = 0;
+
 	/**
 	* Returns the identifier for the engine with the given root directory.
 	*
diff --git a/Engine/Source/Editor/GameProjectGeneration/Private/SProjectBrowser.cpp b/Engine/Source/Editor/GameProjectGeneration/Private/SProjectBrowser.cpp
--- a/Engine/Source/Editor/GameProjectGeneration/Private/SProjectBrowser.cpp
+++ b/Engine/Source/Editor/GameProjectGeneration/Private/SProjectBrowser.cpp
@@ -5,6 +5,7 @@
 #include "MainFrame.h"
 #include "DesktopPlatformModule.h"
 #include "SVerbChoiceDialog.h"
+#include "UProjectInfo.h"
 
 #define LOCTEXT_NAMESPACE "ProjectBrowser"
 
@@ -615,79 +616,100 @@ FReply SProjectBrowser::FindProjects()
 	const FText MyProjectsCategoryName = LOCTEXT("MyProjectsCategoryName", "My Projects");
 	const FText SamplesCategoryName = LOCTEXT("SamplesCategoryName", "Samples");
 
-	// Form a list of all known project files.
-	TMap<FString, FText> DiscoveredProjectFilesToCategory;
+	// Create a map of parent project folders to their category
+	TMap<FString, FText> ParentProjectFoldersToCategory;
 
-	// Start with recents
-	for (auto RecentIt = GEditor->GetGameAgnosticSettings().RecentlyOpenedProjectFiles.CreateConstIterator(); RecentIt; ++RecentIt)
+	// Add the default creation path, in case you've never created a project before or want to include this path anyway
+	ParentProjectFoldersToCategory.Add(GameProjectUtils::GetDefaultProjectCreationPath(), MyProjectsCategoryName);
+
+	// Add in every path that the user has ever created a project file. This is to catch new projects showing up in the user's project folders
+	const TArray<FString> &CreatedProjectPaths = GEditor->GetGameAgnosticSettings().CreatedProjectPaths;
+	for(int32 Idx = 0; Idx < CreatedProjectPaths.Num(); Idx++)
 	{
-		const FString File = IFileManager::Get().ConvertToAbsolutePathForExternalAppForRead(**RecentIt);
-		DiscoveredProjectFilesToCategory.Add(File, MyProjectsCategoryName);
+		FString CreatedProjectPath = CreatedProjectPaths[Idx];
+		FPaths::NormalizeDirectoryName(CreatedProjectPath);
+		ParentProjectFoldersToCategory.Add(CreatedProjectPath, MyProjectsCategoryName);
 	}
 
-	// Form a list of folders that may contain project files and their category names.
-	struct FFolderToCategory
+	// Create a map of project folders to their category
+	TMap<FString, FText> ProjectFoldersToCategory;
+
+	// Find all the subdirectories of all the parent folders
+	for(TMap<FString, FText>::TConstIterator Iter(ParentProjectFoldersToCategory); Iter; ++Iter)
 	{
-		FString Folder;
-		FText Category;
+		const FString &ParentProjectFolder = Iter.Key();
 
-		FFolderToCategory(const FString& InFolder, const FText& InCategory)
-			: Folder(InFolder), Category(InCategory)
-		{}
-	};
+		TArray<FString> ProjectFolders;
+		IFileManager::Get().FindFiles(ProjectFolders, *(ParentProjectFolder / TEXT("*")), false, true);
 
-	TArray<FFolderToCategory> ProjectFileDiscoverFoldersToCategory;
+		for(int32 Idx = 0; Idx < ProjectFolders.Num(); Idx++)
+		{
+			ProjectFoldersToCategory.Add(ParentProjectFolder / ProjectFolders[Idx], Iter.Value());
+		}
+	}
 
-	// Add the default creation path, in case you've never created a project before or want to include this path anyway
-	ProjectFileDiscoverFoldersToCategory.Add(FFolderToCategory(GameProjectUtils::GetDefaultProjectCreationPath(), MyProjectsCategoryName));
+	// Add all the launcher installed sample folders
+	TArray<FString> LauncherSampleDirectories;
+	FDesktopPlatformModule::Get()->EnumerateLauncherSampleInstallations(LauncherSampleDirectories);
+	for(int32 Idx = 0; Idx < LauncherSampleDirectories.Num(); Idx++)
+	{
+		ProjectFoldersToCategory.Add(LauncherSampleDirectories[Idx], SamplesCategoryName);
+	}
 
-	// Add in every path that the user has ever created a project file. This is to catch new projects showing up in the user's project folders
-	for (auto CreatedPathIt = GEditor->GetGameAgnosticSettings().CreatedProjectPaths.CreateConstIterator(); CreatedPathIt; ++CreatedPathIt)
+	// Create a map of all the project files to their category
+	TMap<FString, FText> ProjectFilesToCategory;
+
+	// Add all the folders that the user has opened recently
+	const TArray<FString> &RecentlyOpenedProjectFiles = GEditor->GetGameAgnosticSettings().RecentlyOpenedProjectFiles;
+	for (int32 Idx = 0; Idx < RecentlyOpenedProjectFiles.Num(); Idx++)
 	{
-		ProjectFileDiscoverFoldersToCategory.Add(FFolderToCategory(*CreatedPathIt, MyProjectsCategoryName));
+		FString RecentlyOpenedProjectFile = RecentlyOpenedProjectFiles[Idx];
+		FPaths::NormalizeFilename(RecentlyOpenedProjectFile);
+		ProjectFilesToCategory.Add(RecentlyOpenedProjectFile, MyProjectsCategoryName);
 	}
 
-	// @todo discover projects in the "holding tank" perhaps
-	// Hard-coding a few paths for now
-	ProjectFileDiscoverFoldersToCategory.Add(FFolderToCategory(FPaths::RootDir(), MyProjectsCategoryName));
-	ProjectFileDiscoverFoldersToCategory.Add(FFolderToCategory(FPaths::Combine(*FPaths::RootDir(), TEXT("Samples")), SamplesCategoryName));
-	ProjectFileDiscoverFoldersToCategory.Add(FFolderToCategory(FPaths::Combine(*FPaths::RootDir(), TEXT("Samples"), TEXT("SampleGames")), SamplesCategoryName));
-	ProjectFileDiscoverFoldersToCategory.Add(FFolderToCategory(FPaths::Combine(*FPaths::RootDir(), TEXT("Samples"), TEXT("Showcases")), SamplesCategoryName));
-	ProjectFileDiscoverFoldersToCategory.Add(FFolderToCategory(FPaths::Combine(*FPaths::RootDir(), TEXT("Samples"), TEXT("Demolets")), SamplesCategoryName));
+	// Scan the project folders for project files
+	FString ProjectWildcard = FString::Printf(TEXT("*.%s"), *IProjectManager::GetProjectFileExtension());
+	for(TMap<FString, FText>::TConstIterator Iter(ProjectFoldersToCategory); Iter; ++Iter)
+	{
+		const FString &ProjectFolder = Iter.Key();
+
+		TArray<FString> ProjectFiles;
+		IFileManager::Get().FindFiles(ProjectFiles, *(ProjectFolder / ProjectWildcard), true, false);
+
+		for(int32 Idx = 0; Idx < ProjectFiles.Num(); Idx++)
+		{
+			ProjectFilesToCategory.Add(ProjectFolder / ProjectFiles[Idx], Iter.Value());
+		}
+	}
 
-	// Discover all unique project files in all discover folders
-	for (auto RootFolderIt = ProjectFileDiscoverFoldersToCategory.CreateConstIterator(); RootFolderIt; ++RootFolderIt)
+	// Add all the discovered non-foreign project files
+	const TArray<FString> &NonForeignProjectFiles = FUProjectDictionary::GetDefault().GetProjectPaths();
+	for(int32 Idx = 0; Idx < NonForeignProjectFiles.Num(); Idx++)
 	{
-		const FString Root = (*RootFolderIt).Folder;
-		const FText Category = (*RootFolderIt).Category;
-		const FString RootSearchString = Root / TEXT("*");
-		TArray<FString> PotentialProjectFolders;
-		IFileManager::Get().FindFiles(PotentialProjectFolders, *RootSearchString, /*Files=*/false, /*Directories=*/true);
-		for (auto FolderIt = PotentialProjectFolders.CreateConstIterator(); FolderIt; ++FolderIt)
+		if(!NonForeignProjectFiles[Idx].Contains(TEXT("/Templates/")))
 		{
-			const FString FolderName = Root / (*FolderIt);
-			const FString ProjectSearchString = FolderName / TEXT("*.") + IProjectManager::GetProjectFileExtension();
-			TArray<FString> FoundProjectFiles;
-			IFileManager::Get().FindFiles(FoundProjectFiles, *ProjectSearchString, /*Files=*/true, /*Directories=*/false);
-			for (auto ProjectFilenameIt = FoundProjectFiles.CreateConstIterator(); ProjectFilenameIt; ++ProjectFilenameIt)
-			{
-				const FString PotentiallyRelativeFile = FolderName / (*ProjectFilenameIt);
-				const FString AbsoluteFile = IFileManager::Get().ConvertToAbsolutePathForExternalAppForRead(*PotentiallyRelativeFile);
-				DiscoveredProjectFilesToCategory.Add(AbsoluteFile, Category);
-			}
+			const FText &CategoryName = NonForeignProjectFiles[Idx].Contains(TEXT("/Samples/"))? SamplesCategoryName : MyProjectsCategoryName;
+			ProjectFilesToCategory.Add(NonForeignProjectFiles[Idx], CategoryName);
 		}
 	}
 
-	const FString EngineIdentifier = FDesktopPlatformModule::Get()->GetCurrentEngineIdentifier();
+	// Normalize all the filenames and make sure there are no duplicates
+	TMap<FString, FText> AbsoluteProjectFilesToCategory;
+	for(TMap<FString, FText>::TConstIterator Iter(ProjectFilesToCategory); Iter; ++Iter)
+	{
+		FString AbsoluteFile = IFileManager::Get().ConvertToAbsolutePathForExternalAppForRead(*Iter.Key());
+		AbsoluteProjectFilesToCategory.Add(AbsoluteFile, Iter.Value());
+	}
 
-	// Add all discovered projects to the list
-	for ( auto ProjectFilenameIt = DiscoveredProjectFilesToCategory.CreateConstIterator(); ProjectFilenameIt; ++ProjectFilenameIt )
+	// Add all the discovered projects to the list
+	const FString EngineIdentifier = FDesktopPlatformModule::Get()->GetCurrentEngineIdentifier();
+	for(TMap<FString, FText>::TConstIterator Iter(AbsoluteProjectFilesToCategory); Iter; ++Iter)
 	{
-		const FString& ProjectFilename = ProjectFilenameIt.Key();
-		const FText& DetectedCategory = ProjectFilenameIt.Value();
+		const FString& ProjectFilename = *Iter.Key();
+		const FText& DetectedCategory = Iter.Value();
 		if ( FPaths::FileExists(ProjectFilename) )
 		{
-			const bool bPromptIfSavedWithNewerVersionOfEngine = false;
 			FProjectStatus ProjectStatus;
 			if (IProjectManager::Get().QueryStatusForProject(ProjectFilename, ProjectStatus))
 			{
diff --git a/Engine/Source/Runtime/Core/Private/Misc/UProjectInfo.cpp b/Engine/Source/Runtime/Core/Private/Misc/UProjectInfo.cpp
--- a/Engine/Source/Runtime/Core/Private/Misc/UProjectInfo.cpp
+++ b/Engine/Source/Runtime/Core/Private/Misc/UProjectInfo.cpp
@@ -118,6 +118,13 @@ FString FUProjectDictionary::GetRelativeProjectPathForGame(const TCHAR* InGameNa
 	return TEXT("");
 }
 
+TArray<FString> FUProjectDictionary::GetProjectPaths() const
+{
+	TArray<FString> Paths;
+	ShortProjectNameDictionary.GenerateValueArray(Paths);
+	return Paths;
+}
+
 FUProjectDictionary& FUProjectDictionary::GetDefault()
 {
 	static FUProjectDictionary DefaultDictionary(FPaths::RootDir());
diff --git a/Engine/Source/Runtime/Core/Public/Misc/UProjectInfo.h b/Engine/Source/Runtime/Core/Public/Misc/UProjectInfo.h
--- a/Engine/Source/Runtime/Core/Public/Misc/UProjectInfo.h
+++ b/Engine/Source/Runtime/Core/Public/Misc/UProjectInfo.h
@@ -15,6 +15,9 @@ public:
 	/** Gets the project filename for the given game. Empty if not found. */
 	FString GetRelativeProjectPathForGame(const TCHAR* GameName, const FString& BaseDir) const;
 
+	/** Gets a list of all the known projects */
+	TArray<FString> GetProjectPaths() const;
+
 	/** Gets the project dictionary for the active engine installation */
 	static FUProjectDictionary& GetDefault();
 
