Fix non-unity build error.

diff --git a/Engine/Source/Editor/GameProjectGeneration/Private/SProjectBrowser.cpp b/Engine/Source/Editor/GameProjectGeneration/Private/SProjectBrowser.cpp
--- a/Engine/Source/Editor/GameProjectGeneration/Private/SProjectBrowser.cpp
+++ b/Engine/Source/Editor/GameProjectGeneration/Private/SProjectBrowser.cpp
@@ -6,6 +6,7 @@
 #include "DesktopPlatformModule.h"
 #include "SVerbChoiceDialog.h"
 #include "UProjectInfo.h"
+#include "SourceCodeNavigation.h"
 
 #define LOCTEXT_NAMESPACE "ProjectBrowser"
 
