Improved the way the class wizard informs users about whether a class with be Public or Private

TTP# 332794 - EDITOR: Usability improvements for the New Class Wizard

It now displays a checkbox for Public and Private, with a tooltip explaining what each option means. If you're not using a Public/Private layout for your code, it disables these options and also displays a tooltip explaining why.

Clicking on one of these check boxes will adjust the current path to place the class either in the Public or Private folder, preserving any sub-path information correctly.

Finally, when creating a class in the Public or Classes folder, the class wizard will now automatically add the MODULENAME_API macro to the class declaration.

Also improved the button alignment used for the wizard.

diff --git a/Engine/Source/Editor/GameProjectGeneration/Private/GameProjectUtils.cpp b/Engine/Source/Editor/GameProjectGeneration/Private/GameProjectUtils.cpp
--- a/Engine/Source/Editor/GameProjectGeneration/Private/GameProjectUtils.cpp
+++ b/Engine/Source/Editor/GameProjectGeneration/Private/GameProjectUtils.cpp
@@ -1511,25 +1511,18 @@ bool GameProjectUtils::CalculateSourcePaths(const FString& InPath, FString& OutM
 	const FString AbsoluteInPath = FPaths::ConvertRelativePathToFull(InPath) / ""; // Ensure trailing /
 	OutHeaderPath = AbsoluteInPath;
 	OutSourcePath = AbsoluteInPath;
-	OutModuleName.Empty();
 
-	if(!IsValidSourcePath(InPath, true/*bIncludeModuleName*/, OutFailReason))
+	EClassLocation ClassPathLocation = EClassLocation::UserDefined;
+	if(!GetClassLocation(InPath, OutModuleName, ClassPathLocation, OutFailReason))
 	{
 		return false;
 	}
 
-	// We've validated that this path includes a partial match for our module (eg, MyModule, MyModuleEditor, MyModuleClient)
-	// so extract the actual name of the module from the path so that we can generate the internal folder names correctly
 	const FString BaseRootPath = GetSourceRootPath(false/*bIncludeModuleName*/);
-	const int32 ModuleNameStartIndex = BaseRootPath.Len();
-	const int32 ModuleNameEndIndex = AbsoluteInPath.Find(TEXT("/"), ESearchCase::IgnoreCase, ESearchDir::FromStart, ModuleNameStartIndex);
-	check(ModuleNameEndIndex != INDEX_NONE); // this should never happen since AbsoluteInPath ends in a /, and we verified it started with BaseRootPath in IsValidSourcePath
-	OutModuleName = AbsoluteInPath.Mid(ModuleNameStartIndex, ModuleNameEndIndex - ModuleNameStartIndex);
-
 	const FString RootPath = BaseRootPath / OutModuleName / ""; // Ensure trailing /
-	const FString ClassesPath = RootPath / "Classes" / "";		// Ensure trailing /
 	const FString PublicPath = RootPath / "Public" / "";		// Ensure trailing /
 	const FString PrivatePath = RootPath / "Private" / "";		// Ensure trailing /
+	const FString ClassesPath = RootPath / "Classes" / "";		// Ensure trailing /
 
 	// The root path must exist; we will allow the creation of sub-folders, but not the module root!
 	// We ignore this check if the project doesn't already have source code in it, as the module folder won't yet have been created
@@ -1545,37 +1538,85 @@ bool GameProjectUtils::CalculateSourcePaths(const FString& InPath, FString& OutM
 		return false;
 	}
 
-	// If either the Public or Private path exists, and we're in the root, force the header/source file to use one of these folders
-	const bool bPublicPathExists = IFileManager::Get().DirectoryExists(*PublicPath);
-	const bool bPrivatePathExists = IFileManager::Get().DirectoryExists(*PrivatePath);
-	const bool bForceInternalPath = AbsoluteInPath == RootPath && (bPublicPathExists || bPrivatePathExists);
-
 	// The rules for placing header files are as follows:
-	// 1) If InPath is the source root, check to see if there is a Public folder within it, and if so, place the header there
+	// 1) If InPath is the source root, and GetClassLocation has said the class header should be in the Public folder, put it in the Public folder
 	// 2) Otherwise, just place the header at InPath (the default set above)
 	if(AbsoluteInPath == RootPath)
 	{
-		OutHeaderPath = (bPublicPathExists || bForceInternalPath) ? PublicPath : AbsoluteInPath;
+		OutHeaderPath = (ClassPathLocation == EClassLocation::Public) ? PublicPath : AbsoluteInPath;
 	}
 
 	// The rules for placing source files are as follows:
-	// 1) If InPath is the source root, check to see if there is a Private folder within it, and if so, place the source file there
+	// 1) If InPath is the source root, and GetClassLocation has said the class header should be in the Public folder, put the source file in the Private folder
 	// 2) If InPath is contained within the Public or Classes folder of this module, place it in the equivalent path in the Private folder
 	// 3) Otherwise, just place the source file at InPath (the default set above)
 	if(AbsoluteInPath == RootPath)
 	{
-		OutSourcePath = (bPrivatePathExists || bForceInternalPath) ? PrivatePath : AbsoluteInPath;
+		OutSourcePath = (ClassPathLocation == EClassLocation::Public) ? PrivatePath : AbsoluteInPath;
 	}
-	else if(AbsoluteInPath.StartsWith(ClassesPath))
+	else if(ClassPathLocation == EClassLocation::Public)
+	{
+		OutSourcePath = AbsoluteInPath.Replace(*PublicPath, *PrivatePath);
+	}
+	else if(ClassPathLocation == EClassLocation::Classes)
 	{
 		OutSourcePath = AbsoluteInPath.Replace(*ClassesPath, *PrivatePath);
 	}
+
+	return !OutHeaderPath.IsEmpty() && !OutSourcePath.IsEmpty();
+}
+
+bool GameProjectUtils::GetClassLocation(const FString& InPath, FString& OutModuleName, EClassLocation& OutClassLocation, FText* const OutFailReason)
+{
+	const FString AbsoluteInPath = FPaths::ConvertRelativePathToFull(InPath) / ""; // Ensure trailing /
+	OutModuleName.Empty();
+	OutClassLocation = EClassLocation::UserDefined;
+
+	if(!IsValidSourcePath(InPath, true/*bIncludeModuleName*/, OutFailReason))
+	{
+		return false;
+	}
+
+	// We've validated that this path includes a partial match for our module (eg, MyModule, MyModuleEditor, MyModuleClient)
+	// so extract the actual name of the module from the path so that we can generate the internal folder names correctly
+	const FString BaseRootPath = GetSourceRootPath(false/*bIncludeModuleName*/);
+	const int32 ModuleNameStartIndex = BaseRootPath.Len();
+	const int32 ModuleNameEndIndex = AbsoluteInPath.Find(TEXT("/"), ESearchCase::IgnoreCase, ESearchDir::FromStart, ModuleNameStartIndex);
+	check(ModuleNameEndIndex != INDEX_NONE); // this should never happen since AbsoluteInPath ends in a /, and we verified it started with BaseRootPath in IsValidSourcePath
+	OutModuleName = AbsoluteInPath.Mid(ModuleNameStartIndex, ModuleNameEndIndex - ModuleNameStartIndex);
+
+	const FString RootPath = BaseRootPath / OutModuleName / ""; // Ensure trailing /
+	const FString PublicPath = RootPath / "Public" / "";		// Ensure trailing /
+	const FString PrivatePath = RootPath / "Private" / "";		// Ensure trailing /
+	const FString ClassesPath = RootPath / "Classes" / "";		// Ensure trailing /
+
+	// If either the Public or Private path exists, and we're in the root, force the header/source file to use one of these folders
+	const bool bPublicPathExists = IFileManager::Get().DirectoryExists(*PublicPath);
+	const bool bPrivatePathExists = IFileManager::Get().DirectoryExists(*PrivatePath);
+	const bool bForceInternalPath = AbsoluteInPath == RootPath && (bPublicPathExists || bPrivatePathExists);
+
+	if(AbsoluteInPath == RootPath)
+	{
+		OutClassLocation = (bPublicPathExists || bForceInternalPath) ? EClassLocation::Public : EClassLocation::UserDefined;
+	}
 	else if(AbsoluteInPath.StartsWith(PublicPath))
 	{
-		OutSourcePath = AbsoluteInPath.Replace(*PublicPath, *PrivatePath);
+		OutClassLocation = EClassLocation::Public;
+	}
+	else if(AbsoluteInPath.StartsWith(PrivatePath))
+	{
+		OutClassLocation = EClassLocation::Private;
+	}
+	else if(AbsoluteInPath.StartsWith(ClassesPath))
+	{
+		OutClassLocation = EClassLocation::Classes;
+	}
+	else
+	{
+		OutClassLocation = EClassLocation::UserDefined;
 	}
 
-	return !OutHeaderPath.IsEmpty() && !OutSourcePath.IsEmpty();
+	return true;
 }
 
 bool GameProjectUtils::DuplicateProjectForUpgrade( const FString& InProjectFile, FString &OutNewProjectFile )
@@ -1770,9 +1811,23 @@ bool GameProjectUtils::GenerateClassHeaderFile(const FString& NewHeaderFileName,
 		BaseClassIncludeDirective = FString::Printf(LINE_TERMINATOR TEXT("#include \"%s\""), *BaseClass->GetMetaData(TEXT("IncludePath")));
 	}
 
+	FString ModuleAPIMacro;
+	{
+		FString ModuleName;
+		EClassLocation ClassPathLocation = EClassLocation::UserDefined;
+		GetClassLocation(NewHeaderFileName, ModuleName, ClassPathLocation);
+
+		// If this class isn't Private, make sure and include the API macro so it can be linked within other modules
+		if ( ClassPathLocation != EClassLocation::Private )
+		{
+			ModuleAPIMacro = ModuleName.ToUpper() + "_API "; // include a trailing space for the template formatting
+		}
+	}
+
 	const FString UnprefixedClassName = PrefixedClassName.Mid(1);
 	FString FinalOutput = Template.Replace(TEXT("%COPYRIGHT_LINE%"), *MakeCopyrightLine(), ESearchCase::CaseSensitive);
 	FinalOutput = FinalOutput.Replace(TEXT("%UNPREFIXED_CLASS_NAME%"), *UnprefixedClassName, ESearchCase::CaseSensitive);
+	FinalOutput = FinalOutput.Replace(TEXT("%CLASS_MODULE_API_MACRO%"), *ModuleAPIMacro, ESearchCase::CaseSensitive);
 	FinalOutput = FinalOutput.Replace(TEXT("%UCLASS_SPECIFIER_LIST%"), *MakeCommaDelimitedList(ClassSpecifierList, false), ESearchCase::CaseSensitive);
 	FinalOutput = FinalOutput.Replace(TEXT("%PREFIXED_CLASS_NAME%"), *PrefixedClassName, ESearchCase::CaseSensitive);
 	FinalOutput = FinalOutput.Replace(TEXT("%PREFIXED_BASE_CLASS_NAME%"), *PrefixedBaseClassName, ESearchCase::CaseSensitive);
diff --git a/Engine/Source/Editor/GameProjectGeneration/Private/GameProjectUtils.h b/Engine/Source/Editor/GameProjectGeneration/Private/GameProjectUtils.h
--- a/Engine/Source/Editor/GameProjectGeneration/Private/GameProjectUtils.h
+++ b/Engine/Source/Editor/GameProjectGeneration/Private/GameProjectUtils.h
@@ -5,6 +5,22 @@
 class GameProjectUtils
 {
 public:
+	/** Where is this class located within the Source folder? */
+	enum class EClassLocation : uint8
+	{
+		/** The class is going to a user defined location (outside of the Public, Private, or Classes) folder for this module */
+		UserDefined,
+
+		/** The class is going to the Public folder for this module */
+		Public,
+
+		/** The class is going to the Private folder for this module */
+		Private,
+
+		/** The class is going to the Classes folder for this module */
+		Classes,
+	};
+
 	/** Returns true if the project filename is properly formed and does not conflict with another project */
 	static bool IsValidProjectFileForCreation(const FString& ProjectFile, FText& OutFailReason);
 
@@ -92,6 +108,18 @@ public:
 	 */
 	static bool CalculateSourcePaths(const FString& InPath, FString& OutModuleName, FString& OutHeaderPath, FString& OutSourcePath, FText* const OutFailReason = nullptr);
 
+	/** 
+	 * Given the path provided, work out where it's located within the Source folder
+	 *
+	 * @param	InPath				The path to use a base
+	 * @param	OutModuleName		The module name extracted from the path (the part after GetSourceRootPath(false))
+	 * @param	OutClassLocation	The location within the Source folder
+	 * @param	OutFailReason		Optional parameter to fill with failure information
+	 * 
+	 * @return	false if the paths are invalid
+	 */
+	static bool GetClassLocation(const FString& InPath, FString& OutModuleName, EClassLocation& OutClassLocation, FText* const OutFailReason = nullptr);
+
 	/** Creates a copy of a project directory in order to upgrade it. */
 	static bool DuplicateProjectForUpgrade( const FString& InProjectFile, FString &OutNewProjectFile );
 
diff --git a/Engine/Source/Editor/GameProjectGeneration/Private/SNewClassDialog.cpp b/Engine/Source/Editor/GameProjectGeneration/Private/SNewClassDialog.cpp
--- a/Engine/Source/Editor/GameProjectGeneration/Private/SNewClassDialog.cpp
+++ b/Engine/Source/Editor/GameProjectGeneration/Private/SNewClassDialog.cpp
@@ -58,6 +58,7 @@ BEGIN_SLATE_FUNCTION_BUILD_OPTIMIZATION
 void SNewClassDialog::Construct( const FArguments& InArgs )
 {
 	NewClassPath = GameProjectUtils::GetSourceRootPath(true/*bIncludeModuleName*/);
+	ClassLocation = GameProjectUtils::EClassLocation::UserDefined; // the first call to UpdateInputValidity will set this correctly based on NewClassPath
 
 	ParentClass = InArgs._Class;
 
@@ -271,7 +272,7 @@ void SNewClassDialog::Construct( const FArguments& InArgs )
 
 							+SVerticalBox::Slot()
 							.AutoHeight()
-							.Padding(0, 0, 0, 20)
+							.Padding(0, 0, 0, 2)
 							[
 								SNew(STextBlock)
 								.Text( LOCTEXT("ClassNameDetails", "When you click the \"Create\" button below, a header (.h) file and a source (.cpp) file will be made using this name.") )
@@ -297,125 +298,205 @@ void SNewClassDialog::Construct( const FArguments& InArgs )
 								]
 							]
 
+							// Properties
 							+SVerticalBox::Slot()
 							.AutoHeight()
-							.Padding(0)
 							[
-								SNew(SGridPanel)
-								.FillColumn(1, 1.0f)
-
-								// Name label
-								+SGridPanel::Slot(0, 0)
-								.VAlign(VAlign_Center)
-								.Padding(0, 0, 12, 0)
+								SNew(SBorder)
+								.BorderImage(FEditorStyle::GetBrush("DetailsView.CategoryTop"))
+								.BorderBackgroundColor(FLinearColor(0.6f, 0.6f, 0.6f, 1.0f ))
+								.Padding(FMargin(6.0f, 4.0f, 7.0f, 4.0f))
 								[
-									SNew(STextBlock)
-									.TextStyle( FEditorStyle::Get(), "NewClassDialog.SelectedParentClassLabel" )
-									.Text( LOCTEXT( "NameLabel", "Name" ) )
-								]
+									SNew(SVerticalBox)
 
-								// Name edit box
-								+SGridPanel::Slot(1, 0)
-								.Padding(0.0f, 3.0f)
-								.VAlign(VAlign_Center)
-								[
-									SNew(SBox)
-									.HeightOverride(EditableTextHeight)
+									+SVerticalBox::Slot()
+									.AutoHeight()
+									.Padding(0)
 									[
-										SAssignNew( ClassNameEditBox, SEditableTextBox)
-										.Text( this, &SNewClassDialog::OnGetClassNameText )
-										.OnTextChanged( this, &SNewClassDialog::OnClassNameTextChanged )
-									]
-								]
+										SNew(SGridPanel)
+										.FillColumn(1, 1.0f)
 
-								// Path label
-								+SGridPanel::Slot(0, 1)
-								.VAlign(VAlign_Center)
-								.Padding(0, 0, 12, 0)
-								[
-									SNew(STextBlock)
-									.TextStyle( FEditorStyle::Get(), "NewClassDialog.SelectedParentClassLabel" )
-									.Text( LOCTEXT( "PathLabel", "Path" ).ToString() )
-								]
+										// Name label
+										+SGridPanel::Slot(0, 0)
+										.VAlign(VAlign_Center)
+										.Padding(0, 0, 12, 0)
+										[
+											SNew(STextBlock)
+											.TextStyle( FEditorStyle::Get(), "NewClassDialog.SelectedParentClassLabel" )
+											.Text( LOCTEXT( "NameLabel", "Name" ) )
+										]
 
-								// Path edit box
-								+SGridPanel::Slot(1, 1)
-								.Padding(0.0f, 3.0f)
-								.VAlign(VAlign_Center)
-								[
-									SNew(SBox)
-									.HeightOverride(EditableTextHeight)
-									[
-										SNew(SHorizontalBox)
+										// Name edit box
+										+SGridPanel::Slot(1, 0)
+										.Padding(0.0f, 3.0f)
+										.VAlign(VAlign_Center)
+										[
+											SNew(SBox)
+											.HeightOverride(EditableTextHeight)
+											[
+												SNew(SHorizontalBox)
+
+												+SHorizontalBox::Slot()
+												.FillWidth(1.0f)
+												[
+													SAssignNew( ClassNameEditBox, SEditableTextBox)
+													.Text( this, &SNewClassDialog::OnGetClassNameText )
+													.OnTextChanged( this, &SNewClassDialog::OnClassNameTextChanged )
+												]
+
+												+SHorizontalBox::Slot()
+												.AutoWidth()
+												.Padding(6.0f, 0.0f, 0.0f, 0.0f)
+												[
+													SNew(SHorizontalBox)
+
+													+SHorizontalBox::Slot()
+													.AutoWidth()
+													[
+														SNew(SCheckBox)
+														.Style(FEditorStyle::Get(), "Property.ToggleButton.Start")
+														.IsEnabled(this, &SNewClassDialog::CanChangeClassLocation)
+														.IsChecked(this, &SNewClassDialog::IsClassLocationActive, GameProjectUtils::EClassLocation::Public)
+														.OnCheckStateChanged(this, &SNewClassDialog::OnClassLocationChanged, GameProjectUtils::EClassLocation::Public)
+														.ToolTipText(this, &SNewClassDialog::GetClassLocationTooltip, GameProjectUtils::EClassLocation::Public)
+														[
+															SNew(SBox)
+															.VAlign(VAlign_Center)
+															.HAlign(HAlign_Left)
+															.Padding(FMargin(4.0f, 0.0f, 3.0f, 0.0f))
+															[
+																SNew(STextBlock)
+																.Text(LOCTEXT("Public", "Public"))
+																.ColorAndOpacity(this, &SNewClassDialog::GetClassLocationTextColor, GameProjectUtils::EClassLocation::Public)
+															]
+														]
+													]
+
+													+SHorizontalBox::Slot()
+													.AutoWidth()
+													[
+														SNew(SCheckBox)
+														.Style(FEditorStyle::Get(), "Property.ToggleButton.End")
+														.IsEnabled(this, &SNewClassDialog::CanChangeClassLocation)
+														.IsChecked(this, &SNewClassDialog::IsClassLocationActive, GameProjectUtils::EClassLocation::Private)
+														.OnCheckStateChanged(this, &SNewClassDialog::OnClassLocationChanged, GameProjectUtils::EClassLocation::Private)
+														.ToolTipText(this, &SNewClassDialog::GetClassLocationTooltip, GameProjectUtils::EClassLocation::Private)
+														[
+															SNew(SBox)
+															.VAlign(VAlign_Center)
+															.HAlign(HAlign_Right)
+															.Padding(FMargin(3.0f, 0.0f, 4.0f, 0.0f))
+															[
+																SNew(STextBlock)
+																.Text(LOCTEXT("Private", "Private"))
+																.ColorAndOpacity(this, &SNewClassDialog::GetClassLocationTextColor, GameProjectUtils::EClassLocation::Private)
+															]
+														]
+													]
+												]
+											]
+										]
 
-										+SHorizontalBox::Slot()
-										.FillWidth(1.0f)
+										// Path label
+										+SGridPanel::Slot(0, 1)
+										.VAlign(VAlign_Center)
+										.Padding(0, 0, 12, 0)
 										[
-											SNew(SEditableTextBox)
-											.Text(this, &SNewClassDialog::OnGetClassPathText)
-											.OnTextChanged(this, &SNewClassDialog::OnClassPathTextChanged)
+											SNew(STextBlock)
+											.TextStyle( FEditorStyle::Get(), "NewClassDialog.SelectedParentClassLabel" )
+											.Text( LOCTEXT( "PathLabel", "Path" ).ToString() )
 										]
 
-										+SHorizontalBox::Slot()
-										.AutoWidth()
-										.Padding(6.0f, 1.0f, 0.0f, 0.0f)
+										// Path edit box
+										+SGridPanel::Slot(1, 1)
+										.Padding(0.0f, 3.0f)
+										.VAlign(VAlign_Center)
 										[
-											SNew(SButton)
-											.VAlign(VAlign_Center)
-											.OnClicked(this, &SNewClassDialog::HandleChooseFolderButtonClicked)
-											.Text( LOCTEXT( "BrowseButtonText", "Choose Folder" ) )
+											SNew(SBox)
+											.HeightOverride(EditableTextHeight)
+											[
+												SNew(SHorizontalBox)
+
+												+SHorizontalBox::Slot()
+												.FillWidth(1.0f)
+												[
+													SNew(SEditableTextBox)
+													.Text(this, &SNewClassDialog::OnGetClassPathText)
+													.OnTextChanged(this, &SNewClassDialog::OnClassPathTextChanged)
+												]
+
+												+SHorizontalBox::Slot()
+												.AutoWidth()
+												.Padding(6.0f, 1.0f, 0.0f, 0.0f)
+												[
+													SNew(SButton)
+													.VAlign(VAlign_Center)
+													.OnClicked(this, &SNewClassDialog::HandleChooseFolderButtonClicked)
+													.Text( LOCTEXT( "BrowseButtonText", "Choose Folder" ) )
+												]
+											]
 										]
-									]
-								]
 
-								// Header output label
-								+SGridPanel::Slot(0, 2)
-								.VAlign(VAlign_Center)
-								.Padding(0, 0, 12, 0)
-								[
-									SNew(STextBlock)
-									.TextStyle( FEditorStyle::Get(), "NewClassDialog.SelectedParentClassLabel" )
-									.Text( LOCTEXT( "HeaderFileLabel", "Header File" ).ToString() )
-								]
+										// Header output label
+										+SGridPanel::Slot(0, 2)
+										.VAlign(VAlign_Center)
+										.Padding(0, 0, 12, 0)
+										[
+											SNew(STextBlock)
+											.TextStyle( FEditorStyle::Get(), "NewClassDialog.SelectedParentClassLabel" )
+											.Text( LOCTEXT( "HeaderFileLabel", "Header File" ).ToString() )
+										]
 
-								// Header output text
-								+SGridPanel::Slot(1, 2)
-								.Padding(0.0f, 3.0f)
-								.VAlign(VAlign_Center)
-								[
-									SNew(SBox)
-									.VAlign(VAlign_Center)
-									.HeightOverride(EditableTextHeight)
-									[
-										SNew(STextBlock)
-										.Text(this, &SNewClassDialog::OnGetClassHeaderFileText)
-									]
-								]
+										// Header output text
+										+SGridPanel::Slot(1, 2)
+										.Padding(0.0f, 3.0f)
+										.VAlign(VAlign_Center)
+										[
+											SNew(SBox)
+											.VAlign(VAlign_Center)
+											.HeightOverride(EditableTextHeight)
+											[
+												SNew(STextBlock)
+												.Text(this, &SNewClassDialog::OnGetClassHeaderFileText)
+											]
+										]
 
-								// Source output label
-								+SGridPanel::Slot(0, 3)
-								.VAlign(VAlign_Center)
-								.Padding(0, 0, 12, 0)
-								[
-									SNew(STextBlock)
-									.TextStyle( FEditorStyle::Get(), "NewClassDialog.SelectedParentClassLabel" )
-									.Text( LOCTEXT( "SourceFileLabel", "Source File" ).ToString() )
-								]
+										// Source output label
+										+SGridPanel::Slot(0, 3)
+										.VAlign(VAlign_Center)
+										.Padding(0, 0, 12, 0)
+										[
+											SNew(STextBlock)
+											.TextStyle( FEditorStyle::Get(), "NewClassDialog.SelectedParentClassLabel" )
+											.Text( LOCTEXT( "SourceFileLabel", "Source File" ).ToString() )
+										]
 
-								// Source output text
-								+SGridPanel::Slot(1, 3)
-								.Padding(0.0f, 3.0f)
-								.VAlign(VAlign_Center)
-								[
-									SNew(SBox)
-									.VAlign(VAlign_Center)
-									.HeightOverride(EditableTextHeight)
-									[
-										SNew(STextBlock)
-										.Text(this, &SNewClassDialog::OnGetClassSourceFileText)
+										// Source output text
+										+SGridPanel::Slot(1, 3)
+										.Padding(0.0f, 3.0f)
+										.VAlign(VAlign_Center)
+										[
+											SNew(SBox)
+											.VAlign(VAlign_Center)
+											.HeightOverride(EditableTextHeight)
+											[
+												SNew(STextBlock)
+												.Text(this, &SNewClassDialog::OnGetClassSourceFileText)
+											]
+										]
 									]
 								]
 							]
+
+							+SVerticalBox::Slot()
+							.AutoHeight()
+							.Padding(0.0f)
+							[
+								SNew(SBorder)
+								.Padding(FMargin(0.0f, 3.0f, 0.0f, 0.0f))
+								.BorderImage(FEditorStyle::GetBrush("DetailsView.CategoryBottom"))
+								.BorderBackgroundColor(FLinearColor(0.6f, 0.6f, 0.6f, 1.0f ))
+							]
 						]
 					]
 				]
@@ -790,6 +871,90 @@ FReply SNewClassDialog::HandleChooseFolderButtonClicked()
 	return FReply::Handled();
 }
 
+FSlateColor SNewClassDialog::GetClassLocationTextColor(GameProjectUtils::EClassLocation InLocation) const
+{
+	return (ClassLocation == InLocation) ? FSlateColor(FLinearColor(0, 0, 0)) : FSlateColor(FLinearColor(0.72f, 0.72f, 0.72f, 1.f));
+}
+
+FText SNewClassDialog::GetClassLocationTooltip(GameProjectUtils::EClassLocation InLocation) const
+{
+	if(CanChangeClassLocation())
+	{
+		switch(InLocation)
+		{
+		case GameProjectUtils::EClassLocation::Public:
+			return LOCTEXT("ClassLocation_Public", "A public class can be included and used inside other modules in addition to the module it resides in");
+
+		case GameProjectUtils::EClassLocation::Private:
+			return LOCTEXT("ClassLocation_Private", "A private class can only be included and used within the module it resides in");
+
+		default:
+			break;
+		}
+	}
+
+	return LOCTEXT("ClassLocation_UserDefined", "Your project is either not using a Public and Private source layout, or you're explicitly creating your class outside of the Public or Private folder");
+}
+
+ESlateCheckBoxState::Type SNewClassDialog::IsClassLocationActive(GameProjectUtils::EClassLocation InLocation) const
+{
+	return (ClassLocation == InLocation) ? ESlateCheckBoxState::Checked : ESlateCheckBoxState::Unchecked;
+}
+
+void SNewClassDialog::OnClassLocationChanged(ESlateCheckBoxState::Type InCheckedState, GameProjectUtils::EClassLocation InLocation)
+{
+	if(InCheckedState == ESlateCheckBoxState::Checked)
+	{
+		const FString AbsoluteClassPath = FPaths::ConvertRelativePathToFull(NewClassPath) / ""; // Ensure trailing /
+
+		FString ModuleName;
+		GameProjectUtils::EClassLocation TmpClassLocation = GameProjectUtils::EClassLocation::UserDefined;
+		GameProjectUtils::GetClassLocation(AbsoluteClassPath, ModuleName, TmpClassLocation);
+
+		const FString BaseRootPath = GameProjectUtils::GetSourceRootPath(false/*bIncludeModuleName*/);
+		const FString RootPath = BaseRootPath / ModuleName / "";	// Ensure trailing /
+		const FString PublicPath = RootPath / "Public" / "";		// Ensure trailing /
+		const FString PrivatePath = RootPath / "Private" / "";		// Ensure trailing /
+
+		// Update the class path to be rooted to the Public or Private folder based on InVisibility
+		switch(InLocation)
+		{
+		case GameProjectUtils::EClassLocation::Public:
+			if(AbsoluteClassPath.StartsWith(PrivatePath))
+			{
+				NewClassPath = AbsoluteClassPath.Replace(*PrivatePath, *PublicPath);
+			}
+			else if(AbsoluteClassPath.StartsWith(RootPath))
+			{
+				NewClassPath = AbsoluteClassPath.Replace(*RootPath, *PublicPath);
+			}
+			break;
+
+		case GameProjectUtils::EClassLocation::Private:
+			if(AbsoluteClassPath.StartsWith(PublicPath))
+			{
+				NewClassPath = AbsoluteClassPath.Replace(*PublicPath, *PrivatePath);
+			}
+			else if(AbsoluteClassPath.StartsWith(RootPath))
+			{
+				NewClassPath = AbsoluteClassPath.Replace(*RootPath, *PrivatePath);
+			}
+			break;
+
+		default:
+			break;
+		}
+
+		// Will update ClassVisibility correctly
+		UpdateInputValidity();
+	}
+}
+
+bool SNewClassDialog::CanChangeClassLocation() const
+{
+	return ClassLocation != GameProjectUtils::EClassLocation::UserDefined;
+}
+
 void SNewClassDialog::UpdateInputValidity()
 {
 	bLastInputValidityCheckSuccessful = true;
@@ -800,6 +965,22 @@ void SNewClassDialog::UpdateInputValidity()
 	CalculatedClassHeaderName /= NewClassName + ".h";
 	CalculatedClassSourceName /= NewClassName + ".cpp";
 
+	// If the source paths check as succeeded, check to see if we're using a Public/Private class
+	if(bLastInputValidityCheckSuccessful)
+	{
+		GameProjectUtils::GetClassLocation(NewClassPath, ModuleName, ClassLocation);
+
+		// We only care about the Public and Private folders
+		if(ClassLocation != GameProjectUtils::EClassLocation::Public && ClassLocation != GameProjectUtils::EClassLocation::Private)
+		{
+			ClassLocation = GameProjectUtils::EClassLocation::UserDefined;
+		}
+	}
+	else
+	{
+		ClassLocation = GameProjectUtils::EClassLocation::UserDefined;
+	}
+
 	// Validate the class name only if the path is valid
 	if ( bLastInputValidityCheckSuccessful )
 	{
diff --git a/Engine/Source/Editor/GameProjectGeneration/Private/SNewClassDialog.h b/Engine/Source/Editor/GameProjectGeneration/Private/SNewClassDialog.h
--- a/Engine/Source/Editor/GameProjectGeneration/Private/SNewClassDialog.h
+++ b/Engine/Source/Editor/GameProjectGeneration/Private/SNewClassDialog.h
@@ -108,6 +108,22 @@ private:
 	FReply HandleChooseFolderButtonClicked( );
 
 private:
+
+	/** Get the text color to use for the given class location checkbox */
+	FSlateColor GetClassLocationTextColor(GameProjectUtils::EClassLocation InLocation) const;
+
+	/** Get the tooltip to use for the given class location */
+	FText GetClassLocationTooltip(GameProjectUtils::EClassLocation InLocation) const;
+
+	/** Checks to see if the given class location is active based on the current value of NewClassPath */
+	ESlateCheckBoxState::Type IsClassLocationActive(GameProjectUtils::EClassLocation InLocation) const;
+
+	/** Update the value of NewClassPath so that it uses the given class location */
+	void OnClassLocationChanged(ESlateCheckBoxState::Type InCheckedState, GameProjectUtils::EClassLocation InLocation);
+
+	/** Checks to see if the class location can be changed (relies on us having a Public/Private folder layout) */
+	bool CanChangeClassLocation() const;
+
 	/** Checks the current class name/path for validity and updates cached values accordingly */
 	void UpdateInputValidity();
 
@@ -173,4 +189,7 @@ private:
 
 	/** True if the last validity check returned that the class name/path is valid for creation */
 	bool bLastInputValidityCheckSuccessful;
+
+	/** Whether the class should be created as a Public or Private class */
+	GameProjectUtils::EClassLocation ClassLocation;
 };
diff --git a/Engine/Source/Runtime/Slate/Private/Widgets/Navigation/SWizard.cpp b/Engine/Source/Runtime/Slate/Private/Widgets/Navigation/SWizard.cpp
--- a/Engine/Source/Runtime/Slate/Private/Widgets/Navigation/SWizard.cpp
+++ b/Engine/Source/Runtime/Slate/Private/Widgets/Navigation/SWizard.cpp
@@ -72,28 +72,74 @@ void SWizard::Construct( const FArguments& InArgs )
 				[
 					// 'Prev' button
 					SNew(SButton)
+					.HAlign(HAlign_Center)
+					.VAlign(VAlign_Center)
 					.ContentPadding(FCoreStyle::Get().GetMargin("StandardDialog.ContentPadding"))
 					.IsEnabled(this, &SWizard::HandlePrevButtonIsEnabled)
 					.OnClicked(this, &SWizard::HandlePrevButtonClicked)
 					.Visibility(this, &SWizard::HandlePrevButtonVisibility)
 					.ToolTipText(LOCTEXT("PrevButtonTooltip", "Go back to the previous step"))
-					.Text(LOCTEXT("PrevButtonLabel", "< Back"))
+					[
+						SNew(SHorizontalBox)
+
+						+SHorizontalBox::Slot()
+						.Padding(2.0f, 0.0f)
+						.AutoWidth()
+						.VAlign(VAlign_Center)
+						[
+							SNew(SImage)
+							.Image(FCoreStyle::Get().GetBrush("Wizard.BackIcon"))
+							.ColorAndOpacity(FLinearColor(0.05f, 0.05f, 0.05f))
+						]
+
+						+SHorizontalBox::Slot()
+						.AutoWidth()
+						.VAlign(VAlign_Center)
+						[
+							SNew(STextBlock)
+							.Text(LOCTEXT("PrevButtonLabel", "Back"))
+						]
+					]
 				]
 				+SUniformGridPanel::Slot(1,0)
 				[
 					// 'Next' button
 					SNew(SButton)
+					.HAlign(HAlign_Center)
+					.VAlign(VAlign_Center)
 					.ContentPadding(FCoreStyle::Get().GetMargin("StandardDialog.ContentPadding"))
 					.IsEnabled(this, &SWizard::HandleNextButtonIsEnabled)
 					.OnClicked(this, &SWizard::HandleNextButtonClicked)
 					.Visibility(this, &SWizard::HandleNextButtonVisibility)
 					.ToolTipText(LOCTEXT("NextButtonTooltip", "Go to the next step"))
-					.Text(LOCTEXT("NextButtonLabel", "Next >"))
+					[
+						SNew(SHorizontalBox)
+
+						+SHorizontalBox::Slot()
+						.AutoWidth()
+						.VAlign(VAlign_Center)
+						[
+							SNew(STextBlock)
+							.Text(LOCTEXT("NextButtonLabel", "Next"))
+						]
+
+						+SHorizontalBox::Slot()
+						.Padding(2.0f, 0.0f)
+						.AutoWidth()
+						.VAlign(VAlign_Center)
+						[
+							SNew(SImage)
+							.Image(FCoreStyle::Get().GetBrush("Wizard.NextIcon"))
+							.ColorAndOpacity(FLinearColor(0.05f, 0.05f, 0.05f))
+						]
+					]
 				]
 				+SUniformGridPanel::Slot(2,0)
 				[
 					// 'Finish' button
 					SNew(SButton)
+					.HAlign(HAlign_Center)
+					.VAlign(VAlign_Center)
 					.ContentPadding(FCoreStyle::Get().GetMargin("StandardDialog.ContentPadding"))
 					.IsEnabled(InArgs._CanFinish)
 					.OnClicked(this, &SWizard::HandleFinishButtonClicked)
@@ -109,6 +155,8 @@ void SWizard::Construct( const FArguments& InArgs )
 		[
 			// cancel button
 			SNew(SButton)
+			.HAlign(HAlign_Center)
+			.VAlign(VAlign_Center)
 			.ContentPadding(FCoreStyle::Get().GetMargin("StandardDialog.ContentPadding"))
 			.OnClicked(this, &SWizard::HandleCancelButtonClicked)
 			.ToolTipText(LOCTEXT("CancelButtonTooltip", "Cancel this wizard"))
diff --git a/Engine/Source/Runtime/SlateCore/Private/Styling/CoreStyle.cpp b/Engine/Source/Runtime/SlateCore/Private/Styling/CoreStyle.cpp
--- a/Engine/Source/Runtime/SlateCore/Private/Styling/CoreStyle.cpp
+++ b/Engine/Source/Runtime/SlateCore/Private/Styling/CoreStyle.cpp
@@ -1102,6 +1102,12 @@ TSharedRef<ISlateStyle> FCoreStyle::Create( const FName& InStyleSetName )
 		Style->Set("CrashTracker.Record", new IMAGE_BRUSH("CrashTracker/Record", Icon20x20));
 	}
 
+	// Wizard icons
+	{
+		Style->Set("Wizard.BackIcon", new IMAGE_BRUSH("Icons/BackIcon", Icon8x8));
+		Style->Set("Wizard.NextIcon", new IMAGE_BRUSH("Icons/NextIcon", Icon8x8));
+	}
+
 	return Style;
 }
 
