Condensed the class wizard layout so it wastes less space

TTP# 332794 - EDITOR: Usability improvements for the New Class Wizard

diff --git a/Engine/Source/Editor/GameProjectGeneration/Private/GameProjectUtils.cpp b/Engine/Source/Editor/GameProjectGeneration/Private/GameProjectUtils.cpp
--- a/Engine/Source/Editor/GameProjectGeneration/Private/GameProjectUtils.cpp
+++ b/Engine/Source/Editor/GameProjectGeneration/Private/GameProjectUtils.cpp
@@ -484,7 +484,7 @@ void GameProjectUtils::OpenAddCodeToProjectDialog()
 	TSharedRef<SWindow> AddCodeWindow =
 		SNew(SWindow)
 		.Title(LOCTEXT( "AddCodeWindowHeader", "Add Code"))
-		.ClientSize( FVector2D(1280, 720) )
+		.ClientSize( FVector2D(940, 540) )
 		.SizingRule( ESizingRule::FixedSize )
 		.SupportsMinimize(false) .SupportsMaximize(false);
 
diff --git a/Engine/Source/Editor/GameProjectGeneration/Private/SNewClassDialog.cpp b/Engine/Source/Editor/GameProjectGeneration/Private/SNewClassDialog.cpp
--- a/Engine/Source/Editor/GameProjectGeneration/Private/SNewClassDialog.cpp
+++ b/Engine/Source/Editor/GameProjectGeneration/Private/SNewClassDialog.cpp
@@ -62,7 +62,6 @@ void SNewClassDialog::Construct( const FArguments& InArgs )
 
 	ParentClass = InArgs._Class;
 
-	DialogFixedWidth = 900;
 	bShowFullClassTree = false;
 
 	LastPeriodicValidityCheckTime = 0;
@@ -93,444 +92,439 @@ void SNewClassDialog::Construct( const FArguments& InArgs )
 	ChildSlot
 	[
 		SNew(SBorder)
-		.HAlign(HAlign_Center)
+		.Padding(FMargin(26, 8))
 		.BorderImage( FEditorStyle::GetBrush("Docking.Tab.ContentAreaBrush") )
 		[
-			SNew(SBox)
-			.WidthOverride(DialogFixedWidth)
-			.Padding( FMargin(0, 4) )
-			[
-				SNew(SVerticalBox)
+			SNew(SVerticalBox)
 
-				+SVerticalBox::Slot()
-				[
-					SAssignNew( MainWizard, SWizard)
-					.ShowPageList(false)
-					.CanFinish(this, &SNewClassDialog::CanFinish)
-					.FinishButtonText( LOCTEXT("FinishButtonText", "Create Class").ToString() )
-					.FinishButtonToolTip ( LOCTEXT("FinishButtonToolTip", "Creates the code files to add your new class.").ToString() )
-					.OnCanceled(this, &SNewClassDialog::CancelClicked)
-					.OnFinished(this, &SNewClassDialog::FinishClicked)
-					.InitialPageIndex(ParentClass.IsValid() ? 1 : 0)
+			+SVerticalBox::Slot()
+			[
+				SAssignNew( MainWizard, SWizard)
+				.ShowPageList(false)
+				.CanFinish(this, &SNewClassDialog::CanFinish)
+				.FinishButtonText( LOCTEXT("FinishButtonText", "Create Class").ToString() )
+				.FinishButtonToolTip ( LOCTEXT("FinishButtonToolTip", "Creates the code files to add your new class.").ToString() )
+				.OnCanceled(this, &SNewClassDialog::CancelClicked)
+				.OnFinished(this, &SNewClassDialog::FinishClicked)
+				.InitialPageIndex(ParentClass.IsValid() ? 1 : 0)
 				
-					// Choose parent class
-					+SWizard::Page()
+				// Choose parent class
+				+SWizard::Page()
+				[
+					SNew(SVerticalBox)
+
+					// Title
+					+SVerticalBox::Slot()
+					.AutoHeight()
+					.Padding(0, 20, 0, 0)
 					[
-						SNew(SVerticalBox)
+						SNew(STextBlock)
+						.TextStyle( FEditorStyle::Get(), "NewClassDialog.PageTitle" )
+						.Text( LOCTEXT( "ParentClassTitle", "Choose Parent Class" ) )
+					]
 
-						// Title
-						+SVerticalBox::Slot()
-						.AutoHeight()
-						.Padding(0, 20, 0, 0)
+					// Title spacer
+					+SVerticalBox::Slot()
+					.AutoHeight()
+					.Padding(0, 2, 0, 8)
+					[
+						SNew(SSeparator)
+					]
+
+					// Page description and view options
+					+SVerticalBox::Slot()
+					.AutoHeight()
+					.Padding(0, 10)
+					[
+						SNew(SHorizontalBox)
+
+						+SHorizontalBox::Slot()
+						.FillWidth(1.f)
+						.VAlign(VAlign_Center)
 						[
 							SNew(STextBlock)
-							.TextStyle( FEditorStyle::Get(), "NewClassDialog.PageTitle" )
-							.Text( LOCTEXT( "ParentClassTitle", "Choose Parent Class" ) )
+							.Text( FText::Format( LOCTEXT("ChooseParentClassDescription", "You are about to add a C++ source code file. To compile these files you must have {0} installed."), FSourceCodeNavigation::GetSuggestedSourceCodeIDE() ) )
 						]
 
-						// Title spacer
-						+SVerticalBox::Slot()
-						.AutoHeight()
-						.Padding(0, 2, 0, 8)
+						// Full tree checkbox
+						+SHorizontalBox::Slot()
+						.AutoWidth()
+						.VAlign(VAlign_Center)
+						.Padding(4, 0, 0, 0)
 						[
-							SNew(SSeparator)
+							SNew(SCheckBox)
+							.IsChecked( this, &SNewClassDialog::IsFullClassTreeChecked )
+							.OnCheckStateChanged( this, &SNewClassDialog::OnFullClassTreeChanged )
+							[
+								SNew(STextBlock)
+								.Text( LOCTEXT( "FullClassTree", "Show All Classes" ) )
+							]
 						]
+					]
 
-						// Page description and view options
-						+SVerticalBox::Slot()
-						.AutoHeight()
-						.Padding(0, 10)
+					// Add Code list
+					+SVerticalBox::Slot()
+					.FillHeight(1.f)
+					.Padding(0, 10)
+					[
+						SNew(SBorder)
+						.BorderImage( FEditorStyle::GetBrush("ToolPanel.GroupBorder") )
 						[
-							SNew(SHorizontalBox)
+							SNew(SVerticalBox)
 
-							+SHorizontalBox::Slot()
-							.FillWidth(1.f)
-							.VAlign(VAlign_Center)
+							+SVerticalBox::Slot()
 							[
-								SNew(STextBlock)
-								.Text( FText::Format( LOCTEXT("ChooseParentClassDescription", "You are about to add a C++ source code file. To compile these files you must have {0} installed."), FSourceCodeNavigation::GetSuggestedSourceCodeIDE() ) )
+								// Basic view
+								SAssignNew(ParentClassListView, SListView< TSharedPtr<FParentClassItem> >)
+								.ListItemsSource(&ParentClassItemsSource)
+								.SelectionMode(ESelectionMode::Single)
+								.ClearSelectionOnClick(false)
+								.OnGenerateRow(this, &SNewClassDialog::MakeParentClassListViewWidget)
+								.OnMouseButtonDoubleClick( this, &SNewClassDialog::OnParentClassItemDoubleClicked )
+								.OnSelectionChanged(this, &SNewClassDialog::OnClassSelected)
+								.Visibility(this, &SNewClassDialog::GetBasicParentClassVisibility)
 							]
 
-							// Full tree checkbox
-							+SHorizontalBox::Slot()
-							.AutoWidth()
-							.VAlign(VAlign_Center)
-							.Padding(4, 0, 0, 0)
+							+SVerticalBox::Slot()
 							[
-								SNew(SCheckBox)
-								.IsChecked( this, &SNewClassDialog::IsFullClassTreeChecked )
-								.OnCheckStateChanged( this, &SNewClassDialog::OnFullClassTreeChanged )
+								// Advanced view
+								SNew(SBox)
+								.Visibility(this, &SNewClassDialog::GetAdvancedParentClassVisibility)
 								[
-									SNew(STextBlock)
-									.Text( LOCTEXT( "FullClassTree", "Show All Classes" ) )
+									ClassViewer.ToSharedRef()
 								]
 							]
 						]
+					]
 
-						// Add Code list
-						+SVerticalBox::Slot()
-						.FillHeight(1.f)
-						.Padding(0, 10)
-						[
-							SNew(SBorder)
-							.BorderImage( FEditorStyle::GetBrush("ToolPanel.GroupBorder") )
-							[
-								SNew(SVerticalBox)
-
-								+SVerticalBox::Slot()
-								[
-									// Basic view
-									SAssignNew(ParentClassListView, SListView< TSharedPtr<FParentClassItem> >)
-									.ListItemsSource(&ParentClassItemsSource)
-									.SelectionMode(ESelectionMode::Single)
-									.ClearSelectionOnClick(false)
-									.OnGenerateRow(this, &SNewClassDialog::MakeParentClassListViewWidget)
-									.OnMouseButtonDoubleClick( this, &SNewClassDialog::OnParentClassItemDoubleClicked )
-									.OnSelectionChanged(this, &SNewClassDialog::OnClassSelected)
-									.Visibility(this, &SNewClassDialog::GetBasicParentClassVisibility)
-								]
+					// Class selection
+					+SVerticalBox::Slot()
+					.Padding(30, 2)
+					.AutoHeight()
+					[
+						SNew(SHorizontalBox)
 
-								+SVerticalBox::Slot()
-								[
-									// Advanced view
-									SNew(SBox)
-									.Visibility(this, &SNewClassDialog::GetAdvancedParentClassVisibility)
-									[
-										ClassViewer.ToSharedRef()
-									]
-								]
-							]
+						// Class label
+						+SHorizontalBox::Slot()
+						.AutoWidth()
+						.VAlign(VAlign_Center)
+						.Padding(0, 0, 12, 0)
+						[
+							SNew(STextBlock)
+							.TextStyle( FEditorStyle::Get(), "NewClassDialog.SelectedParentClassLabel" )
+							.Text( LOCTEXT( "ParentClassLabel", "Selected Class" ) )
 						]
 
-						// Class selection
-						+SVerticalBox::Slot()
-						.Padding(40, 2)
-						.AutoHeight()
+						// Class selection preview
+						+SHorizontalBox::Slot()
+						.FillWidth(1.f)
+						.VAlign(VAlign_Center)
 						[
-							SNew(SHorizontalBox)
+							SNew(STextBlock)
+							.Text( this, &SNewClassDialog::GetSelectedParentClassName )
+						]
+					]
+				]
 
-							// Class label
-							+SHorizontalBox::Slot()
-							.AutoWidth()
-							.VAlign(VAlign_Center)
-							.Padding(0, 0, 12, 0)
-							[
-								SNew(STextBlock)
-								.TextStyle( FEditorStyle::Get(), "NewClassDialog.SelectedParentClassLabel" )
-								.Text( LOCTEXT( "ParentClassLabel", "Selected Class" ) )
-							]
+				// Name class
+				+SWizard::Page()
+				.OnEnter(this, &SNewClassDialog::OnNamePageEntered)
+				[
+					SNew(SVerticalBox)
 
-							// Class selection preview
-							+SHorizontalBox::Slot()
-							.FillWidth(1.f)
-							.VAlign(VAlign_Center)
-							[
-								SNew(STextBlock)
-								.Text( this, &SNewClassDialog::GetSelectedParentClassName )
-							]
-						]
+					// Title
+					+SVerticalBox::Slot()
+					.AutoHeight()
+					.Padding(0, 20, 0, 0)
+					[
+						SNew(STextBlock)
+						.TextStyle( FEditorStyle::Get(), "NewClassDialog.PageTitle" )
+						.Text( this, &SNewClassDialog::GetNameClassTitle )
 					]
 
-					// Name class
-					+SWizard::Page()
-					.OnEnter(this, &SNewClassDialog::OnNamePageEntered)
+					// Title spacer
+					+SVerticalBox::Slot()
+					.AutoHeight()
+					.Padding(0, 2, 0, 8)
+					[
+						SNew(SSeparator)
+					]
+
+					+SVerticalBox::Slot()
+					.FillHeight(1.f)
+					.Padding(0, 10)
 					[
 						SNew(SVerticalBox)
 
-						// Title
 						+SVerticalBox::Slot()
 						.AutoHeight()
-						.Padding(0, 20, 0, 0)
+						.Padding(0, 0, 0, 5)
 						[
 							SNew(STextBlock)
-							.TextStyle( FEditorStyle::Get(), "NewClassDialog.PageTitle" )
-							.Text( this, &SNewClassDialog::GetNameClassTitle )
+							.Text( LOCTEXT("ClassNameDescription", "Enter a name for your new class. Class names may only contain alphanumeric characters, and may not contain a space.") )
 						]
 
-						// Title spacer
 						+SVerticalBox::Slot()
 						.AutoHeight()
-						.Padding(0, 2, 0, 0)
+						.Padding(0, 0, 0, 2)
 						[
-							SNew(SSeparator)
+							SNew(STextBlock)
+							.Text( LOCTEXT("ClassNameDetails", "When you click the \"Create\" button below, a header (.h) file and a source (.cpp) file will be made using this name.") )
 						]
 
+						// Name Error label
 						+SVerticalBox::Slot()
-						.FillHeight(1.f)
-						.Padding(80, 2)
-						.VAlign(VAlign_Center)
+						.AutoHeight()
+						.Padding(0, 5)
 						[
-							SNew(SVerticalBox)
-
-							+SVerticalBox::Slot()
-							.AutoHeight()
-							.Padding(0, 0, 0, 5)
+							// Constant height, whether the label is visible or not
+							SNew(SBox).HeightOverride(20)
 							[
-								SNew(STextBlock)
-								.Text( LOCTEXT("ClassNameDescription", "Enter a name for your new class. Class names may only contain alphanumeric characters, and may not contain a space.") )
+								SNew(SBorder)
+								.Visibility( this, &SNewClassDialog::GetNameErrorLabelVisibility )
+								.BorderImage( FEditorStyle::GetBrush("NewClassDialog.ErrorLabelBorder") )
+								.Content()
+								[
+									SNew(STextBlock)
+									.Text( this, &SNewClassDialog::GetNameErrorLabelText )
+									.TextStyle( FEditorStyle::Get(), "NewClassDialog.ErrorLabelFont" )
+								]
 							]
+						]
 
-							+SVerticalBox::Slot()
-							.AutoHeight()
-							.Padding(0, 0, 0, 2)
+						// Properties
+						+SVerticalBox::Slot()
+						.AutoHeight()
+						[
+							SNew(SBorder)
+							.BorderImage(FEditorStyle::GetBrush("DetailsView.CategoryTop"))
+							.BorderBackgroundColor(FLinearColor(0.6f, 0.6f, 0.6f, 1.0f ))
+							.Padding(FMargin(6.0f, 4.0f, 7.0f, 4.0f))
 							[
-								SNew(STextBlock)
-								.Text( LOCTEXT("ClassNameDetails", "When you click the \"Create\" button below, a header (.h) file and a source (.cpp) file will be made using this name.") )
-							]
+								SNew(SVerticalBox)
 
-							// Name Error label
-							+SVerticalBox::Slot()
-							.AutoHeight()
-							.Padding(0, 5)
-							[
-								// Constant height, whether the label is visible or not
-								SNew(SBox).HeightOverride(20)
+								+SVerticalBox::Slot()
+								.AutoHeight()
+								.Padding(0)
 								[
-									SNew(SBorder)
-									.Visibility( this, &SNewClassDialog::GetNameErrorLabelVisibility )
-									.BorderImage( FEditorStyle::GetBrush("NewClassDialog.ErrorLabelBorder") )
-									.Content()
+									SNew(SGridPanel)
+									.FillColumn(1, 1.0f)
+
+									// Name label
+									+SGridPanel::Slot(0, 0)
+									.VAlign(VAlign_Center)
+									.Padding(0, 0, 12, 0)
 									[
 										SNew(STextBlock)
-										.Text( this, &SNewClassDialog::GetNameErrorLabelText )
-										.TextStyle( FEditorStyle::Get(), "NewClassDialog.ErrorLabelFont" )
+										.TextStyle( FEditorStyle::Get(), "NewClassDialog.SelectedParentClassLabel" )
+										.Text( LOCTEXT( "NameLabel", "Name" ) )
 									]
-								]
-							]
 
-							// Properties
-							+SVerticalBox::Slot()
-							.AutoHeight()
-							[
-								SNew(SBorder)
-								.BorderImage(FEditorStyle::GetBrush("DetailsView.CategoryTop"))
-								.BorderBackgroundColor(FLinearColor(0.6f, 0.6f, 0.6f, 1.0f ))
-								.Padding(FMargin(6.0f, 4.0f, 7.0f, 4.0f))
-								[
-									SNew(SVerticalBox)
-
-									+SVerticalBox::Slot()
-									.AutoHeight()
-									.Padding(0)
+									// Name edit box
+									+SGridPanel::Slot(1, 0)
+									.Padding(0.0f, 3.0f)
+									.VAlign(VAlign_Center)
 									[
-										SNew(SGridPanel)
-										.FillColumn(1, 1.0f)
-
-										// Name label
-										+SGridPanel::Slot(0, 0)
-										.VAlign(VAlign_Center)
-										.Padding(0, 0, 12, 0)
+										SNew(SBox)
+										.HeightOverride(EditableTextHeight)
 										[
-											SNew(STextBlock)
-											.TextStyle( FEditorStyle::Get(), "NewClassDialog.SelectedParentClassLabel" )
-											.Text( LOCTEXT( "NameLabel", "Name" ) )
-										]
+											SNew(SHorizontalBox)
 
-										// Name edit box
-										+SGridPanel::Slot(1, 0)
-										.Padding(0.0f, 3.0f)
-										.VAlign(VAlign_Center)
-										[
-											SNew(SBox)
-											.HeightOverride(EditableTextHeight)
+											+SHorizontalBox::Slot()
+											.FillWidth(1.0f)
 											[
-												SNew(SHorizontalBox)
+												SAssignNew( ClassNameEditBox, SEditableTextBox)
+												.Text( this, &SNewClassDialog::OnGetClassNameText )
+												.OnTextChanged( this, &SNewClassDialog::OnClassNameTextChanged )
+											]
 
-												+SHorizontalBox::Slot()
-												.FillWidth(1.0f)
-												[
-													SAssignNew( ClassNameEditBox, SEditableTextBox)
-													.Text( this, &SNewClassDialog::OnGetClassNameText )
-													.OnTextChanged( this, &SNewClassDialog::OnClassNameTextChanged )
-												]
+											+SHorizontalBox::Slot()
+											.AutoWidth()
+											.Padding(6.0f, 0.0f, 0.0f, 0.0f)
+											[
+												SNew(SHorizontalBox)
 
 												+SHorizontalBox::Slot()
 												.AutoWidth()
-												.Padding(6.0f, 0.0f, 0.0f, 0.0f)
 												[
-													SNew(SHorizontalBox)
-
-													+SHorizontalBox::Slot()
-													.AutoWidth()
+													SNew(SCheckBox)
+													.Style(FEditorStyle::Get(), "Property.ToggleButton.Start")
+													.IsEnabled(this, &SNewClassDialog::CanChangeClassLocation)
+													.IsChecked(this, &SNewClassDialog::IsClassLocationActive, GameProjectUtils::EClassLocation::Public)
+													.OnCheckStateChanged(this, &SNewClassDialog::OnClassLocationChanged, GameProjectUtils::EClassLocation::Public)
+													.ToolTipText(this, &SNewClassDialog::GetClassLocationTooltip, GameProjectUtils::EClassLocation::Public)
 													[
-														SNew(SCheckBox)
-														.Style(FEditorStyle::Get(), "Property.ToggleButton.Start")
-														.IsEnabled(this, &SNewClassDialog::CanChangeClassLocation)
-														.IsChecked(this, &SNewClassDialog::IsClassLocationActive, GameProjectUtils::EClassLocation::Public)
-														.OnCheckStateChanged(this, &SNewClassDialog::OnClassLocationChanged, GameProjectUtils::EClassLocation::Public)
-														.ToolTipText(this, &SNewClassDialog::GetClassLocationTooltip, GameProjectUtils::EClassLocation::Public)
+														SNew(SBox)
+														.VAlign(VAlign_Center)
+														.HAlign(HAlign_Left)
+														.Padding(FMargin(4.0f, 0.0f, 3.0f, 0.0f))
 														[
-															SNew(SBox)
-															.VAlign(VAlign_Center)
-															.HAlign(HAlign_Left)
-															.Padding(FMargin(4.0f, 0.0f, 3.0f, 0.0f))
-															[
-																SNew(STextBlock)
-																.Text(LOCTEXT("Public", "Public"))
-																.ColorAndOpacity(this, &SNewClassDialog::GetClassLocationTextColor, GameProjectUtils::EClassLocation::Public)
-															]
+															SNew(STextBlock)
+															.Text(LOCTEXT("Public", "Public"))
+															.ColorAndOpacity(this, &SNewClassDialog::GetClassLocationTextColor, GameProjectUtils::EClassLocation::Public)
 														]
 													]
+												]
 
-													+SHorizontalBox::Slot()
-													.AutoWidth()
+												+SHorizontalBox::Slot()
+												.AutoWidth()
+												[
+													SNew(SCheckBox)
+													.Style(FEditorStyle::Get(), "Property.ToggleButton.End")
+													.IsEnabled(this, &SNewClassDialog::CanChangeClassLocation)
+													.IsChecked(this, &SNewClassDialog::IsClassLocationActive, GameProjectUtils::EClassLocation::Private)
+													.OnCheckStateChanged(this, &SNewClassDialog::OnClassLocationChanged, GameProjectUtils::EClassLocation::Private)
+													.ToolTipText(this, &SNewClassDialog::GetClassLocationTooltip, GameProjectUtils::EClassLocation::Private)
 													[
-														SNew(SCheckBox)
-														.Style(FEditorStyle::Get(), "Property.ToggleButton.End")
-														.IsEnabled(this, &SNewClassDialog::CanChangeClassLocation)
-														.IsChecked(this, &SNewClassDialog::IsClassLocationActive, GameProjectUtils::EClassLocation::Private)
-														.OnCheckStateChanged(this, &SNewClassDialog::OnClassLocationChanged, GameProjectUtils::EClassLocation::Private)
-														.ToolTipText(this, &SNewClassDialog::GetClassLocationTooltip, GameProjectUtils::EClassLocation::Private)
+														SNew(SBox)
+														.VAlign(VAlign_Center)
+														.HAlign(HAlign_Right)
+														.Padding(FMargin(3.0f, 0.0f, 4.0f, 0.0f))
 														[
-															SNew(SBox)
-															.VAlign(VAlign_Center)
-															.HAlign(HAlign_Right)
-															.Padding(FMargin(3.0f, 0.0f, 4.0f, 0.0f))
-															[
-																SNew(STextBlock)
-																.Text(LOCTEXT("Private", "Private"))
-																.ColorAndOpacity(this, &SNewClassDialog::GetClassLocationTextColor, GameProjectUtils::EClassLocation::Private)
-															]
+															SNew(STextBlock)
+															.Text(LOCTEXT("Private", "Private"))
+															.ColorAndOpacity(this, &SNewClassDialog::GetClassLocationTextColor, GameProjectUtils::EClassLocation::Private)
 														]
 													]
 												]
 											]
 										]
+									]
 
-										// Path label
-										+SGridPanel::Slot(0, 1)
-										.VAlign(VAlign_Center)
-										.Padding(0, 0, 12, 0)
-										[
-											SNew(STextBlock)
-											.TextStyle( FEditorStyle::Get(), "NewClassDialog.SelectedParentClassLabel" )
-											.Text( LOCTEXT( "PathLabel", "Path" ).ToString() )
-										]
+									// Path label
+									+SGridPanel::Slot(0, 1)
+									.VAlign(VAlign_Center)
+									.Padding(0, 0, 12, 0)
+									[
+										SNew(STextBlock)
+										.TextStyle( FEditorStyle::Get(), "NewClassDialog.SelectedParentClassLabel" )
+										.Text( LOCTEXT( "PathLabel", "Path" ).ToString() )
+									]
 
-										// Path edit box
-										+SGridPanel::Slot(1, 1)
-										.Padding(0.0f, 3.0f)
-										.VAlign(VAlign_Center)
+									// Path edit box
+									+SGridPanel::Slot(1, 1)
+									.Padding(0.0f, 3.0f)
+									.VAlign(VAlign_Center)
+									[
+										SNew(SBox)
+										.HeightOverride(EditableTextHeight)
 										[
-											SNew(SBox)
-											.HeightOverride(EditableTextHeight)
-											[
-												SNew(SHorizontalBox)
+											SNew(SHorizontalBox)
 
-												+SHorizontalBox::Slot()
-												.FillWidth(1.0f)
-												[
-													SNew(SEditableTextBox)
-													.Text(this, &SNewClassDialog::OnGetClassPathText)
-													.OnTextChanged(this, &SNewClassDialog::OnClassPathTextChanged)
-												]
-
-												+SHorizontalBox::Slot()
-												.AutoWidth()
-												.Padding(6.0f, 1.0f, 0.0f, 0.0f)
-												[
-													SNew(SButton)
-													.VAlign(VAlign_Center)
-													.OnClicked(this, &SNewClassDialog::HandleChooseFolderButtonClicked)
-													.Text( LOCTEXT( "BrowseButtonText", "Choose Folder" ) )
-												]
+											+SHorizontalBox::Slot()
+											.FillWidth(1.0f)
+											[
+												SNew(SEditableTextBox)
+												.Text(this, &SNewClassDialog::OnGetClassPathText)
+												.OnTextChanged(this, &SNewClassDialog::OnClassPathTextChanged)
 											]
-										]
 
-										// Header output label
-										+SGridPanel::Slot(0, 2)
-										.VAlign(VAlign_Center)
-										.Padding(0, 0, 12, 0)
-										[
-											SNew(STextBlock)
-											.TextStyle( FEditorStyle::Get(), "NewClassDialog.SelectedParentClassLabel" )
-											.Text( LOCTEXT( "HeaderFileLabel", "Header File" ).ToString() )
-										]
-
-										// Header output text
-										+SGridPanel::Slot(1, 2)
-										.Padding(0.0f, 3.0f)
-										.VAlign(VAlign_Center)
-										[
-											SNew(SBox)
-											.VAlign(VAlign_Center)
-											.HeightOverride(EditableTextHeight)
+											+SHorizontalBox::Slot()
+											.AutoWidth()
+											.Padding(6.0f, 1.0f, 0.0f, 0.0f)
 											[
-												SNew(STextBlock)
-												.Text(this, &SNewClassDialog::OnGetClassHeaderFileText)
+												SNew(SButton)
+												.VAlign(VAlign_Center)
+												.OnClicked(this, &SNewClassDialog::HandleChooseFolderButtonClicked)
+												.Text( LOCTEXT( "BrowseButtonText", "Choose Folder" ) )
 											]
 										]
+									]
 
-										// Source output label
-										+SGridPanel::Slot(0, 3)
+									// Header output label
+									+SGridPanel::Slot(0, 2)
+									.VAlign(VAlign_Center)
+									.Padding(0, 0, 12, 0)
+									[
+										SNew(STextBlock)
+										.TextStyle( FEditorStyle::Get(), "NewClassDialog.SelectedParentClassLabel" )
+										.Text( LOCTEXT( "HeaderFileLabel", "Header File" ).ToString() )
+									]
+
+									// Header output text
+									+SGridPanel::Slot(1, 2)
+									.Padding(0.0f, 3.0f)
+									.VAlign(VAlign_Center)
+									[
+										SNew(SBox)
 										.VAlign(VAlign_Center)
-										.Padding(0, 0, 12, 0)
+										.HeightOverride(EditableTextHeight)
 										[
 											SNew(STextBlock)
-											.TextStyle( FEditorStyle::Get(), "NewClassDialog.SelectedParentClassLabel" )
-											.Text( LOCTEXT( "SourceFileLabel", "Source File" ).ToString() )
+											.Text(this, &SNewClassDialog::OnGetClassHeaderFileText)
 										]
+									]
 
-										// Source output text
-										+SGridPanel::Slot(1, 3)
-										.Padding(0.0f, 3.0f)
+									// Source output label
+									+SGridPanel::Slot(0, 3)
+									.VAlign(VAlign_Center)
+									.Padding(0, 0, 12, 0)
+									[
+										SNew(STextBlock)
+										.TextStyle( FEditorStyle::Get(), "NewClassDialog.SelectedParentClassLabel" )
+										.Text( LOCTEXT( "SourceFileLabel", "Source File" ).ToString() )
+									]
+
+									// Source output text
+									+SGridPanel::Slot(1, 3)
+									.Padding(0.0f, 3.0f)
+									.VAlign(VAlign_Center)
+									[
+										SNew(SBox)
 										.VAlign(VAlign_Center)
+										.HeightOverride(EditableTextHeight)
 										[
-											SNew(SBox)
-											.VAlign(VAlign_Center)
-											.HeightOverride(EditableTextHeight)
-											[
-												SNew(STextBlock)
-												.Text(this, &SNewClassDialog::OnGetClassSourceFileText)
-											]
+											SNew(STextBlock)
+											.Text(this, &SNewClassDialog::OnGetClassSourceFileText)
 										]
 									]
 								]
 							]
+						]
 
-							+SVerticalBox::Slot()
-							.AutoHeight()
-							.Padding(0.0f)
-							[
-								SNew(SBorder)
-								.Padding(FMargin(0.0f, 3.0f, 0.0f, 0.0f))
-								.BorderImage(FEditorStyle::GetBrush("DetailsView.CategoryBottom"))
-								.BorderBackgroundColor(FLinearColor(0.6f, 0.6f, 0.6f, 1.0f ))
-							]
+						+SVerticalBox::Slot()
+						.AutoHeight()
+						.Padding(0.0f)
+						[
+							SNew(SBorder)
+							.Padding(FMargin(0.0f, 3.0f, 0.0f, 0.0f))
+							.BorderImage(FEditorStyle::GetBrush("DetailsView.CategoryBottom"))
+							.BorderBackgroundColor(FLinearColor(0.6f, 0.6f, 0.6f, 1.0f ))
 						]
 					]
 				]
+			]
 
-				+ SVerticalBox::Slot()
-					.Padding(0, 5)
-					.AutoHeight()
-					[
-						SNew(SBorder)
-							.Visibility( this, &SNewClassDialog::GetGlobalErrorLabelVisibility )
-							.BorderImage( FEditorStyle::GetBrush("NewClassDialog.ErrorLabelBorder") )
-							.Content()
-							[
-								SNew(SHorizontalBox)
+			// IDE download information
+			+SVerticalBox::Slot()
+			.Padding(0, 5)
+			.AutoHeight()
+			[
+				SNew(SBorder)
+				.Visibility( this, &SNewClassDialog::GetGlobalErrorLabelVisibility )
+				.BorderImage( FEditorStyle::GetBrush("NewClassDialog.ErrorLabelBorder") )
+				.Content()
+				[
+					SNew(SHorizontalBox)
 
-								+ SHorizontalBox::Slot()
-									.VAlign(VAlign_Center)
-									[
-										SNew(STextBlock)
-											.Text( this, &SNewClassDialog::GetGlobalErrorLabelText )
-											.TextStyle( FEditorStyle::Get(), "NewClassDialog.ErrorLabelFont" )
-									]
+					+SHorizontalBox::Slot()
+					.VAlign(VAlign_Center)
+					[
+						SNew(STextBlock)
+						.Text( this, &SNewClassDialog::GetGlobalErrorLabelText )
+						.TextStyle( FEditorStyle::Get(), "NewClassDialog.ErrorLabelFont" )
+					]
 
-								+ SHorizontalBox::Slot()
-									.VAlign(VAlign_Center)
-									.AutoWidth()
-									[
-										SNew(SHyperlink)
-											.Text( FText::Format( LOCTEXT("IDEDownloadLinkText", "Download {0}"), FSourceCodeNavigation::GetSuggestedSourceCodeIDE() ) )
-											.OnNavigate( this, &SNewClassDialog::OnDownloadIDEClicked, FSourceCodeNavigation::GetSuggestedSourceCodeIDEDownloadURL() )
-											.Visibility( this, &SNewClassDialog::GetGlobalErrorLabelIDELinkVisibility )
-									]
-							]
+					+SHorizontalBox::Slot()
+					.VAlign(VAlign_Center)
+					.AutoWidth()
+					[
+						SNew(SHyperlink)
+						.Text( FText::Format( LOCTEXT("IDEDownloadLinkText", "Download {0}"), FSourceCodeNavigation::GetSuggestedSourceCodeIDE() ) )
+						.OnNavigate( this, &SNewClassDialog::OnDownloadIDEClicked, FSourceCodeNavigation::GetSuggestedSourceCodeIDEDownloadURL() )
+						.Visibility( this, &SNewClassDialog::GetGlobalErrorLabelIDELinkVisibility )
 					]
+				]
 			]
 		]
 	];
@@ -578,10 +572,13 @@ TSharedRef<ITableRow> SNewClassDialog::MakeParentClassListViewWidget(TSharedPtr<
 		ClassDescription = ClassDescription.Left(NewLineIndex + 1);
 	}
 
+	// Strip out any new-lines in the description
+	ClassDescription = ClassDescription.Replace(TEXT("\n"), TEXT(" "));
+
 	const FSlateBrush* ClassBrush = FClassIconFinder::FindIconForClass(Class);
 
-	const int32 ItemHeight = 128;
-	const int32 DescriptionIndent = 128;
+	const int32 ItemHeight = 64;
+	const int32 DescriptionIndent = 32;
 	return
 		SNew( STableRow<TSharedPtr<FParentClassItem>>, OwnerTable )
 		.Style(FEditorStyle::Get(), "NewClassDialog.ParentClassListView.TableRow")
@@ -620,7 +617,7 @@ TSharedRef<ITableRow> SNewClassDialog::MakeParentClassListViewWidget(TSharedPtr<
 				.Padding(DescriptionIndent, 0, 0, 0)
 				[
 					SNew(STextBlock)
-					.WrapTextAt(DialogFixedWidth - DescriptionIndent - 32)
+					//.AutoWrapText(true)
 					.Text(ClassDescription)
 				]
 			]
diff --git a/Engine/Source/Editor/GameProjectGeneration/Private/SNewClassDialog.h b/Engine/Source/Editor/GameProjectGeneration/Private/SNewClassDialog.h
--- a/Engine/Source/Editor/GameProjectGeneration/Private/SNewClassDialog.h
+++ b/Engine/Source/Editor/GameProjectGeneration/Private/SNewClassDialog.h
@@ -169,9 +169,6 @@ private:
 	/** The selected parent class */
 	TWeakObjectPtr<UClass> ParentClass;
 
-	/** The fixed width of the dialog */
-	int32 DialogFixedWidth;
-
 	/** If true, the full class tree will be shown in the parent class selection */
 	bool bShowFullClassTree;
 
