Add Code to Project wizard now accounts for subfolders in path when replacing %MODULE_NAME% token in template

diff --git a/Engine/Source/Editor/GameProjectGeneration/Private/GameProjectUtils.cpp b/Engine/Source/Editor/GameProjectGeneration/Private/GameProjectUtils.cpp
--- a/Engine/Source/Editor/GameProjectGeneration/Private/GameProjectUtils.cpp
+++ b/Engine/Source/Editor/GameProjectGeneration/Private/GameProjectUtils.cpp
@@ -1893,10 +1893,18 @@ bool GameProjectUtils::GenerateClassCPPFile(const FString& NewCPPFileName, const
 		PropertyOverridesStr += *PropertyOverrides[OverrideIdx];
 	}
 
+	// account for subfolders in new filename
+	FString ModuleHeaderPath = ModuleName;
+	int32 SubfolderCount = CalculateSubfolderCount(NewCPPFileName, ModuleName);
+	for (int32 i = 0; i < SubfolderCount; ++i)
+	{
+		ModuleHeaderPath = "../" + ModuleHeaderPath;
+	}
+
 	const FString UnprefixedClassName = PrefixedClassName.Mid(1);
 	FString FinalOutput = Template.Replace(TEXT("%COPYRIGHT_LINE%"), *MakeCopyrightLine(), ESearchCase::CaseSensitive);
 	FinalOutput = FinalOutput.Replace(TEXT("%UNPREFIXED_CLASS_NAME%"), *UnprefixedClassName, ESearchCase::CaseSensitive);
-	FinalOutput = FinalOutput.Replace(TEXT("%MODULE_NAME%"), *ModuleName, ESearchCase::CaseSensitive);
+	FinalOutput = FinalOutput.Replace(TEXT("%MODULE_NAME%"), *ModuleHeaderPath, ESearchCase::CaseSensitive);
 	FinalOutput = FinalOutput.Replace(TEXT("%PREFIXED_CLASS_NAME%"), *PrefixedClassName, ESearchCase::CaseSensitive);
 	FinalOutput = FinalOutput.Replace(TEXT("%PROPERTY_OVERRIDES%"), *PropertyOverridesStr, ESearchCase::CaseSensitive);
 	FinalOutput = FinalOutput.Replace(TEXT("%ADDITIONAL_MEMBER_DEFINITIONS%"), *AdditionalMemberDefinitions, ESearchCase::CaseSensitive);
@@ -2321,4 +2329,31 @@ bool GameProjectUtils::AddCodeToProject_Internal(const FString& NewClassName, co
 	return true;
 }
 
+int32 GameProjectUtils::CalculateSubfolderCount(const FString& Filepath, const FString& RelativeRoot)
+{
+	int32 ModuleNameIndex = Filepath.Find(RelativeRoot, ESearchCase::IgnoreCase, ESearchDir::FromEnd);
+	if (ModuleNameIndex == INDEX_NONE)
+	{
+		return 0;
+	}
+
+	FString WorkingString = Filepath.Mid(ModuleNameIndex + RelativeRoot.Len() + 1);
+
+	int Depth = 0;
+	int32 SlashIndex;
+	while (WorkingString.FindChar(TEXT('/'), SlashIndex))
+	{
+		// because input is not sanitized against multiple sequential slashes (e.g. '////')
+		// only consider slashes that don't appear at beginning of the working string
+		if (SlashIndex > 0)
+		{
+			Depth++;
+		}
+
+		WorkingString = WorkingString.Mid(SlashIndex + 1);
+	}
+
+	return Depth;
+}
+
 #undef LOCTEXT_NAMESPACE
diff --git a/Engine/Source/Editor/GameProjectGeneration/Private/GameProjectUtils.h b/Engine/Source/Editor/GameProjectGeneration/Private/GameProjectUtils.h
--- a/Engine/Source/Editor/GameProjectGeneration/Private/GameProjectUtils.h
+++ b/Engine/Source/Editor/GameProjectGeneration/Private/GameProjectUtils.h
@@ -258,6 +258,8 @@ private:
 	/** Handler for the user confirming they've read the name legnth warning */
 	static void OnWarningReasonOk();
 
+	static int32 CalculateSubfolderCount(const FString& Filepath, const FString& RelativeRoot);
+
 private:
 	static TWeakPtr<SNotificationItem> UpdateGameProjectNotification;
 	static TWeakPtr<SNotificationItem> WarningProjectNameNotification;
