Added support for the Class Wizard to create non-UObject classes

TTP# 332794 - EDITOR: Usability improvements for the New Class Wizard

Abstracted away the parent class information into FNewClassInfo, which can hold either a UClass*, or an enum value corresponding to a pre-defined non-UObject type class (currently; an empty class, a Slate widget, and a Slate widget style).

Made the interface of GenerateClassHeaderFile and GenerateClassCPPFile more consistent. They now both take an unprefixed class name as well as the base class information; any extra information they can then generate internally. They also no longer rely on the filename to work out the class name, as the Slate classes may add a prefix or suffix to the filename due to the way the templates are set-up.

Updated everything in the SNewClassDialog and GameProjectUtils that was previously working with a UClass* to work with a FNewClassInfo instead.

diff --git a/Engine/Source/Editor/GameProjectGeneration/Private/GameProjectUtils.cpp b/Engine/Source/Editor/GameProjectGeneration/Private/GameProjectUtils.cpp
--- a/Engine/Source/Editor/GameProjectGeneration/Private/GameProjectUtils.cpp
+++ b/Engine/Source/Editor/GameProjectGeneration/Private/GameProjectUtils.cpp
@@ -13,6 +13,8 @@
 #include "DesktopPlatformModule.h"
 #include "TargetPlatform.h"
 
+#include "ClassIconFinder.h"
+
 #define LOCTEXT_NAMESPACE "GameProjectUtils"
 
 #define MAX_PROJECT_PATH_BUFFER_SPACE 130 // Leave a reasonable buffer of additional characters to account for files created in the content directory during or after project generation
@@ -24,6 +26,254 @@ checkAtCompileTime(PLATFORM_MAX_FILEPATH_LENGTH - MAX_PROJECT_PATH_BUFFER_SPACE
 TWeakPtr<SNotificationItem> GameProjectUtils::UpdateGameProjectNotification = NULL;
 TWeakPtr<SNotificationItem> GameProjectUtils::WarningProjectNameNotification = NULL;
 
+FString GameProjectUtils::FNewClassInfo::GetClassName() const
+{
+	switch(ClassType)
+	{
+	case EClassType::UObject:
+		return BaseClass ? FName::NameToDisplayString(BaseClass->GetName(), false) : TEXT("");
+
+	case EClassType::EmptyCpp:
+		return TEXT("None");
+
+	case EClassType::SlateWidget:
+		return TEXT("Slate Widget");
+
+	case EClassType::SlateWidgetStyle:
+		return TEXT("Slate Widget Style");
+
+	default:
+		break;
+	}
+	return TEXT("");
+}
+
+FString GameProjectUtils::FNewClassInfo::GetClassDescription() const
+{
+	switch(ClassType)
+	{
+	case EClassType::UObject:
+		{
+			if(BaseClass)
+			{
+				FString ClassDescription = BaseClass->GetToolTipText().ToString();
+				int32 FullStopIndex = 0;
+				if(ClassDescription.FindChar('.', FullStopIndex))
+				{
+					// Only show the first sentence so as not to clutter up the UI with a detailed description of implementation details
+					ClassDescription = ClassDescription.Left(FullStopIndex + 1);
+				}
+
+				// Strip out any new-lines in the description
+				ClassDescription = ClassDescription.Replace(TEXT("\n"), TEXT(" "));
+				return ClassDescription;
+			}
+		}
+		break;
+
+	case EClassType::EmptyCpp:
+		return TEXT("An empty C++ class with a default constructor and destructor");
+
+	case EClassType::SlateWidget:
+		return TEXT("A custom Slate widget, deriving from SCompoundWidget");
+
+	case EClassType::SlateWidgetStyle:
+		return TEXT("A custom Slate widget style, deriving from FSlateWidgetStyle, along with its associated UObject wrapper class");
+
+	default:
+		break;
+	}
+	return TEXT("");
+}
+
+const FSlateBrush* GameProjectUtils::FNewClassInfo::GetClassIcon() const
+{
+	// Safe to do even if BaseClass is null, since FindIconForClass will return the default icon
+	return FClassIconFinder::FindIconForClass(BaseClass);
+}
+
+FString GameProjectUtils::FNewClassInfo::GetClassPrefixCPP() const
+{
+	switch(ClassType)
+	{
+	case EClassType::UObject:
+		return BaseClass ? BaseClass->GetPrefixCPP() : TEXT("");
+
+	case EClassType::EmptyCpp:
+		return TEXT("");
+
+	case EClassType::SlateWidget:
+		return TEXT("S");
+
+	case EClassType::SlateWidgetStyle:
+		return TEXT("F");
+
+	default:
+		break;
+	}
+	return TEXT("");
+}
+
+FString GameProjectUtils::FNewClassInfo::GetClassNameCPP() const
+{
+	switch(ClassType)
+	{
+	case EClassType::UObject:
+		return BaseClass ? BaseClass->GetName() : TEXT("");
+
+	case EClassType::EmptyCpp:
+		return TEXT("");
+
+	case EClassType::SlateWidget:
+		return TEXT("CompoundWidget");
+
+	case EClassType::SlateWidgetStyle:
+		return TEXT("SlateWidgetStyle");
+
+	default:
+		break;
+	}
+	return TEXT("");
+}
+
+FString GameProjectUtils::FNewClassInfo::GetCleanClassName(const FString& ClassName) const
+{
+	FString CleanClassName = ClassName;
+
+	switch(ClassType)
+	{
+	case EClassType::SlateWidgetStyle:
+		{
+			// Slate widget style classes always take the form FMyThingWidget, and UMyThingWidgetStyle
+			// if our class ends with either Widget or WidgetStyle, we need to strip those out to avoid silly looking duplicates
+			if(CleanClassName.EndsWith(TEXT("Style")))
+			{
+				CleanClassName = CleanClassName.LeftChop(5); // 5 for "Style"
+			}
+			if(CleanClassName.EndsWith(TEXT("Widget")))
+			{
+				CleanClassName = CleanClassName.LeftChop(6); // 6 for "Widget"
+			}
+		}
+		break;
+
+	default:
+		break;
+	}
+
+	return CleanClassName;
+}
+
+FString GameProjectUtils::FNewClassInfo::GetFinalClassName(const FString& ClassName) const
+{
+	const FString CleanClassName = GetCleanClassName(ClassName);
+
+	switch(ClassType)
+	{
+	case EClassType::SlateWidgetStyle:
+		return FString::Printf(TEXT("%sWidgetStyle"), *CleanClassName);
+
+	default:
+		break;
+	}
+
+	return CleanClassName;
+}
+
+bool GameProjectUtils::FNewClassInfo::GetIncludePath(FString& OutIncludePath) const
+{
+	switch(ClassType)
+	{
+	case EClassType::UObject:
+		if(BaseClass && BaseClass->HasMetaData(TEXT("IncludePath")))
+		{
+			OutIncludePath = BaseClass->GetMetaData(TEXT("IncludePath"));
+			return true;
+		}
+		break;
+
+	default:
+		break;
+	}
+	return false;
+}
+
+FString GameProjectUtils::FNewClassInfo::GetHeaderFilename(const FString& ClassName) const
+{
+	const FString HeaderFilename = GetFinalClassName(ClassName) + TEXT(".h");
+
+	switch(ClassType)
+	{
+	case EClassType::SlateWidget:
+		return TEXT("S") + HeaderFilename;
+
+	default:
+		break;
+	}
+
+	return HeaderFilename;
+}
+
+FString GameProjectUtils::FNewClassInfo::GetSourceFilename(const FString& ClassName) const
+{
+	const FString SourceFilename = GetFinalClassName(ClassName) + TEXT(".cpp");
+
+	switch(ClassType)
+	{
+	case EClassType::SlateWidget:
+		return TEXT("S") + SourceFilename;
+
+	default:
+		break;
+	}
+
+	return SourceFilename;
+}
+
+FString GameProjectUtils::FNewClassInfo::GetHeaderTemplateFilename() const
+{
+	switch(ClassType)
+	{
+	case EClassType::UObject:
+		return TEXT("UObjectClass.h.template");
+
+	case EClassType::EmptyCpp:
+		return TEXT("EmptyClass.h.template");
+
+	case EClassType::SlateWidget:
+		return TEXT("SlateWidget.h.template");
+
+	case EClassType::SlateWidgetStyle:
+		return TEXT("SlateWidgetStyle.h.template");
+
+	default:
+		break;
+	}
+	return TEXT("");
+}
+
+FString GameProjectUtils::FNewClassInfo::GetSourceTemplateFilename() const
+{
+	switch(ClassType)
+	{
+	case EClassType::UObject:
+		return TEXT("UObjectClass.cpp.template");
+
+	case EClassType::EmptyCpp:
+		return TEXT("EmptyClass.cpp.template");
+
+	case EClassType::SlateWidget:
+		return TEXT("SlateWidget.cpp.template");
+
+	case EClassType::SlateWidgetStyle:
+		return TEXT("SlateWidgetStyle.cpp.template");
+
+	default:
+		break;
+	}
+	return TEXT("");
+}
+
 bool GameProjectUtils::IsValidProjectFileForCreation(const FString& ProjectFile, FText& OutFailReason)
 {
 	const FString BaseProjectFile = FPaths::GetBaseFilename(ProjectFile);
@@ -566,14 +816,16 @@ bool GameProjectUtils::IsValidClassNameForCreation(const FString& NewClassName,
 	return true;
 }
 
-bool GameProjectUtils::AddCodeToProject(const FString& NewClassName, const FString& NewClassPath, const UClass* ParentClass, FString& OutHeaderFilePath, FString& OutCppFilePath, FText& OutFailReason)
+bool GameProjectUtils::AddCodeToProject(const FString& NewClassName, const FString& NewClassPath, const FNewClassInfo ParentClassInfo, FString& OutHeaderFilePath, FString& OutCppFilePath, FText& OutFailReason)
 {
-	const bool bAddCodeSuccessful = AddCodeToProject_Internal(NewClassName, NewClassPath, ParentClass, OutHeaderFilePath, OutCppFilePath, OutFailReason);
+	const bool bAddCodeSuccessful = AddCodeToProject_Internal(NewClassName, NewClassPath, ParentClassInfo, OutHeaderFilePath, OutCppFilePath, OutFailReason);
 
 	if( FEngineAnalytics::IsAvailable() )
 	{
+		const FString ParentClassName = ParentClassInfo.GetClassNameCPP();
+
 		TArray<FAnalyticsEventAttribute> EventAttributes;
-		EventAttributes.Add(FAnalyticsEventAttribute(TEXT("ParentClass"), ParentClass ? ParentClass->GetName() : TEXT("None")));
+		EventAttributes.Add(FAnalyticsEventAttribute(TEXT("ParentClass"), ParentClassName.IsEmpty() ? TEXT("None") : ParentClassName));
 		EventAttributes.Add(FAnalyticsEventAttribute(TEXT("Outcome"), bAddCodeSuccessful ? TEXT("Successful") : TEXT("Failed")));
 
 		FEngineAnalytics::GetProvider().RecordEvent( TEXT( "Editor.AddCodeToProject.CodeAdded" ), EventAttributes );
@@ -1360,9 +1612,10 @@ bool GameProjectUtils::GenerateGameFrameworkSourceCode(const FString& NewProject
 	// MyGamePlayerController.h
 	{
 		const UClass* BaseClass = APlayerController::StaticClass();
-		const FString NewHeaderFilename = GameModulePath / NewProjectName + BaseClass->GetName() + TEXT(".h");
+		const FString NewClassName = NewProjectName + BaseClass->GetName();
+		const FString NewHeaderFilename = GameModulePath / NewClassName + TEXT(".h");
 		FString UnusedSyncLocation;
-		if ( GenerateClassHeaderFile(NewHeaderFilename, BaseClass, TArray<FString>(), TEXT(""), TEXT(""), UnusedSyncLocation, OutFailReason) )
+		if ( GenerateClassHeaderFile(NewHeaderFilename, NewClassName, FNewClassInfo(BaseClass), TArray<FString>(), TEXT(""), TEXT(""), UnusedSyncLocation, OutFailReason) )
 		{
 			OutCreatedFiles.Add(NewHeaderFilename);
 		}
@@ -1375,9 +1628,10 @@ bool GameProjectUtils::GenerateGameFrameworkSourceCode(const FString& NewProject
 	// MyGameGameMode.h
 	{
 		const UClass* BaseClass = AGameMode::StaticClass();
-		const FString NewHeaderFilename = GameModulePath / NewProjectName + BaseClass->GetName() + TEXT(".h");
+		const FString NewClassName = NewProjectName + BaseClass->GetName();
+		const FString NewHeaderFilename = GameModulePath / NewClassName + TEXT(".h");
 		FString UnusedSyncLocation;
-		if ( GenerateClassHeaderFile(NewHeaderFilename, BaseClass, TArray<FString>(), TEXT(""), TEXT(""), UnusedSyncLocation, OutFailReason) )
+		if ( GenerateClassHeaderFile(NewHeaderFilename, NewClassName, FNewClassInfo(BaseClass), TArray<FString>(), TEXT(""), TEXT(""), UnusedSyncLocation, OutFailReason) )
 		{
 			OutCreatedFiles.Add(NewHeaderFilename);
 		}
@@ -1391,9 +1645,10 @@ bool GameProjectUtils::GenerateGameFrameworkSourceCode(const FString& NewProject
 	FString PrefixedPlayerControllerClassName;
 	{
 		const UClass* BaseClass = APlayerController::StaticClass();
-		const FString NewCPPFilename = GameModulePath / NewProjectName + BaseClass->GetName() + TEXT(".cpp");
-		PrefixedPlayerControllerClassName = FString(BaseClass->GetPrefixCPP()) + NewProjectName + BaseClass->GetName();
-		if ( GenerateClassCPPFile(NewCPPFilename, NewProjectName, PrefixedPlayerControllerClassName, TArray<FString>(), TArray<FString>(), TEXT(""), OutFailReason) )
+		const FString NewClassName = NewProjectName + BaseClass->GetName();
+		const FString NewCPPFilename = GameModulePath / NewClassName + TEXT(".cpp");
+		PrefixedPlayerControllerClassName = FString(BaseClass->GetPrefixCPP()) + NewClassName;
+		if ( GenerateClassCPPFile(NewCPPFilename, NewClassName, FNewClassInfo(BaseClass), TArray<FString>(), TArray<FString>(), TEXT(""), OutFailReason) )
 		{
 			OutCreatedFiles.Add(NewCPPFilename);
 		}
@@ -1406,8 +1661,8 @@ bool GameProjectUtils::GenerateGameFrameworkSourceCode(const FString& NewProject
 	// MyGameGameMode.cpp
 	{
 		const UClass* BaseClass = AGameMode::StaticClass();
-		const FString NewCPPFilename = GameModulePath / NewProjectName + BaseClass->GetName() + TEXT(".cpp");
-		const FString PrefixedClassName = FString(BaseClass->GetPrefixCPP()) + NewProjectName + BaseClass->GetName();
+		const FString NewClassName = NewProjectName + BaseClass->GetName();
+		const FString NewCPPFilename = GameModulePath / NewClassName + TEXT(".cpp");
 		
 		TArray<FString> PropertyOverrides;
 		PropertyOverrides.Add( FString::Printf( TEXT("PlayerControllerClass = %s::StaticClass();"), *PrefixedPlayerControllerClassName ) );
@@ -1418,7 +1673,7 @@ bool GameProjectUtils::GenerateGameFrameworkSourceCode(const FString& NewProject
 		TArray<FString> AdditionalIncludes;
 		AdditionalIncludes.Add(PlayerControllerClassName);
 
-		if ( GenerateClassCPPFile(NewCPPFilename, NewProjectName, PrefixedClassName, AdditionalIncludes, PropertyOverrides, TEXT(""), OutFailReason) )
+		if ( GenerateClassCPPFile(NewCPPFilename, NewClassName, FNewClassInfo(BaseClass), AdditionalIncludes, PropertyOverrides, TEXT(""), OutFailReason) )
 		{
 			OutCreatedFiles.Add(NewCPPFilename);
 		}
@@ -1792,23 +2047,23 @@ FString GameProjectUtils::MakeIncludeList(const TArray<FString>& InList)
 	return ReturnString;
 }
 
-bool GameProjectUtils::GenerateClassHeaderFile(const FString& NewHeaderFileName, const UClass* BaseClass, const TArray<FString>& ClassSpecifierList, const FString& ClassProperties, const FString& ClassFunctionDeclarations, FString& OutSyncLocation, FText& OutFailReason)
+bool GameProjectUtils::GenerateClassHeaderFile(const FString& NewHeaderFileName, const FString UnPrefixedClassName, const FNewClassInfo ParentClassInfo, const TArray<FString>& ClassSpecifierList, const FString& ClassProperties, const FString& ClassFunctionDeclarations, FString& OutSyncLocation, FText& OutFailReason)
 {
 	FString Template;
-	if ( !ReadTemplateFile(TEXT("UObjectClass.h.template"), Template, OutFailReason) )
+	if ( !ReadTemplateFile(ParentClassInfo.GetHeaderTemplateFilename(), Template, OutFailReason) )
 	{
 		return false;
 	}
 
-	const FString UnPrefixedClassName = FPaths::GetBaseFilename(NewHeaderFileName);
-	const FString ClassPrefix = BaseClass->GetPrefixCPP();
+	const FString ClassPrefix = ParentClassInfo.GetClassPrefixCPP();
 	const FString PrefixedClassName = ClassPrefix + UnPrefixedClassName;
-	const FString PrefixedBaseClassName = ClassPrefix + BaseClass->GetName();
+	const FString PrefixedBaseClassName = ClassPrefix + ParentClassInfo.GetClassNameCPP();
 
 	FString BaseClassIncludeDirective;
-	if(BaseClass->HasMetaData(TEXT("IncludePath")))
+	FString BaseClassIncludePath;
+	if(ParentClassInfo.GetIncludePath(BaseClassIncludePath))
 	{
-		BaseClassIncludeDirective = FString::Printf(LINE_TERMINATOR TEXT("#include \"%s\""), *BaseClass->GetMetaData(TEXT("IncludePath")));
+		BaseClassIncludeDirective = FString::Printf(LINE_TERMINATOR TEXT("#include \"%s\""), *BaseClassIncludePath);
 	}
 
 	FString ModuleAPIMacro;
@@ -1824,9 +2079,9 @@ bool GameProjectUtils::GenerateClassHeaderFile(const FString& NewHeaderFileName,
 		}
 	}
 
-	const FString UnprefixedClassName = PrefixedClassName.Mid(1);
+	// Not all of these will exist in every class template
 	FString FinalOutput = Template.Replace(TEXT("%COPYRIGHT_LINE%"), *MakeCopyrightLine(), ESearchCase::CaseSensitive);
-	FinalOutput = FinalOutput.Replace(TEXT("%UNPREFIXED_CLASS_NAME%"), *UnprefixedClassName, ESearchCase::CaseSensitive);
+	FinalOutput = FinalOutput.Replace(TEXT("%UNPREFIXED_CLASS_NAME%"), *UnPrefixedClassName, ESearchCase::CaseSensitive);
 	FinalOutput = FinalOutput.Replace(TEXT("%CLASS_MODULE_API_MACRO%"), *ModuleAPIMacro, ESearchCase::CaseSensitive);
 	FinalOutput = FinalOutput.Replace(TEXT("%UCLASS_SPECIFIER_LIST%"), *MakeCommaDelimitedList(ClassSpecifierList, false), ESearchCase::CaseSensitive);
 	FinalOutput = FinalOutput.Replace(TEXT("%PREFIXED_CLASS_NAME%"), *PrefixedClassName, ESearchCase::CaseSensitive);
@@ -1862,14 +2117,22 @@ bool GameProjectUtils::GenerateClassHeaderFile(const FString& NewHeaderFileName,
 	return WriteOutputFile(NewHeaderFileName, FinalOutput, OutFailReason);
 }
 
-bool GameProjectUtils::GenerateClassCPPFile(const FString& NewCPPFileName, const FString& ModuleName, const FString& PrefixedClassName, const TArray<FString>& AdditionalIncludes, const TArray<FString>& PropertyOverrides, const FString& AdditionalMemberDefinitions, FText& OutFailReason)
+bool GameProjectUtils::GenerateClassCPPFile(const FString& NewCPPFileName, const FString UnPrefixedClassName, const FNewClassInfo ParentClassInfo, const TArray<FString>& AdditionalIncludes, const TArray<FString>& PropertyOverrides, const FString& AdditionalMemberDefinitions, FText& OutFailReason)
 {
 	FString Template;
-	if ( !ReadTemplateFile(TEXT("UObjectClass.cpp.template"), Template, OutFailReason) )
+	if ( !ReadTemplateFile(ParentClassInfo.GetSourceTemplateFilename(), Template, OutFailReason) )
 	{
 		return false;
 	}
 
+	const FString ClassPrefix = ParentClassInfo.GetClassPrefixCPP();
+	const FString PrefixedClassName = ClassPrefix + UnPrefixedClassName;
+	const FString PrefixedBaseClassName = ClassPrefix + ParentClassInfo.GetClassNameCPP();
+
+	FString ModuleName;
+	EClassLocation ClassPathLocation = EClassLocation::UserDefined;
+	GetClassLocation(NewCPPFileName, ModuleName, ClassPathLocation);
+
 	FString AdditionalIncludesStr;
 	for (int32 IncludeIdx = 0; IncludeIdx < AdditionalIncludes.Num(); ++IncludeIdx)
 	{
@@ -1901,9 +2164,9 @@ bool GameProjectUtils::GenerateClassCPPFile(const FString& NewCPPFileName, const
 		ModuleHeaderPath = "../" + ModuleHeaderPath;
 	}
 
-	const FString UnprefixedClassName = PrefixedClassName.Mid(1);
+	// Not all of these will exist in every class template
 	FString FinalOutput = Template.Replace(TEXT("%COPYRIGHT_LINE%"), *MakeCopyrightLine(), ESearchCase::CaseSensitive);
-	FinalOutput = FinalOutput.Replace(TEXT("%UNPREFIXED_CLASS_NAME%"), *UnprefixedClassName, ESearchCase::CaseSensitive);
+	FinalOutput = FinalOutput.Replace(TEXT("%UNPREFIXED_CLASS_NAME%"), *UnPrefixedClassName, ESearchCase::CaseSensitive);
 	FinalOutput = FinalOutput.Replace(TEXT("%MODULE_NAME%"), *ModuleHeaderPath, ESearchCase::CaseSensitive);
 	FinalOutput = FinalOutput.Replace(TEXT("%PREFIXED_CLASS_NAME%"), *PrefixedClassName, ESearchCase::CaseSensitive);
 	FinalOutput = FinalOutput.Replace(TEXT("%PROPERTY_OVERRIDES%"), *PropertyOverridesStr, ESearchCase::CaseSensitive);
@@ -2221,15 +2484,18 @@ bool GameProjectUtils::ProjectHasCodeFiles()
 	return GameProjectUtils::GetProjectCodeFileCount() > 0;
 }
 
-bool GameProjectUtils::AddCodeToProject_Internal(const FString& NewClassName, const FString& NewClassPath, const UClass* ParentClass, FString& OutHeaderFilePath, FString& OutCppFilePath, FText& OutFailReason)
+bool GameProjectUtils::AddCodeToProject_Internal(const FString& NewClassName, const FString& NewClassPath, const FNewClassInfo ParentClassInfo, FString& OutHeaderFilePath, FString& OutCppFilePath, FText& OutFailReason)
 {
-	if ( !ParentClass )
+	if ( !ParentClassInfo.IsSet() )
 	{
 		OutFailReason = LOCTEXT("NoParentClass", "You must specify a parent class");
 		return false;
 	}
 
-	if ( !IsValidClassNameForCreation(NewClassName, OutFailReason) )
+	const FString CleanClassName = ParentClassInfo.GetCleanClassName(NewClassName);
+	const FString FinalClassName = ParentClassInfo.GetFinalClassName(NewClassName);
+
+	if ( !IsValidClassNameForCreation(FinalClassName, OutFailReason) )
 	{
 		return false;
 	}
@@ -2259,10 +2525,10 @@ bool GameProjectUtils::AddCodeToProject_Internal(const FString& NewClassName, co
 		const FString SourceDir = FPaths::GameSourceDir().LeftChop(1); // Trim the trailing /
 
 		// Assuming the game name is the same as the primary game module name
-		const FString ModuleName = FApp::GetGameName();
+		const FString GameModuleName = FApp::GetGameName();
 
 		TArray<FString> StartupModuleNames;
-		if ( GenerateBasicSourceCode(SourceDir, ModuleName, StartupModuleNames, CreatedFiles, OutFailReason) )
+		if ( GenerateBasicSourceCode(SourceDir, GameModuleName, StartupModuleNames, CreatedFiles, OutFailReason) )
 		{
 			UpdateProject(&StartupModuleNames);
 		}
@@ -2275,9 +2541,9 @@ bool GameProjectUtils::AddCodeToProject_Internal(const FString& NewClassName, co
 
 	// Class Header File
 	FString SyncLocation;
-	const FString NewHeaderFilename = NewHeaderPath / NewClassName + TEXT(".h");
+	const FString NewHeaderFilename = NewHeaderPath / ParentClassInfo.GetHeaderFilename(NewClassName);
 	{
-		if ( GenerateClassHeaderFile(NewHeaderFilename, ParentClass, TArray<FString>(), TEXT(""), TEXT(""), SyncLocation, OutFailReason) )
+		if ( GenerateClassHeaderFile(NewHeaderFilename, CleanClassName, ParentClassInfo, TArray<FString>(), TEXT(""), TEXT(""), SyncLocation, OutFailReason) )
 		{
 			CreatedFiles.Add(NewHeaderFilename);
 		}
@@ -2289,10 +2555,9 @@ bool GameProjectUtils::AddCodeToProject_Internal(const FString& NewClassName, co
 	}
 
 	// Class CPP file
-	const FString NewCppFilename = NewCppPath / NewClassName + TEXT(".cpp");
+	const FString NewCppFilename = NewCppPath / ParentClassInfo.GetSourceFilename(NewClassName);
 	{
-		const FString PrefixedClassName = FString(ParentClass->GetPrefixCPP()) + NewClassName;
-		if ( GenerateClassCPPFile(NewCppFilename, ModuleName, PrefixedClassName, TArray<FString>(), TArray<FString>(), TEXT(""), OutFailReason) )
+		if ( GenerateClassCPPFile(NewCppFilename, CleanClassName, ParentClassInfo, TArray<FString>(), TArray<FString>(), TEXT(""), OutFailReason) )
 		{
 			CreatedFiles.Add(NewCppFilename);
 		}
diff --git a/Engine/Source/Editor/GameProjectGeneration/Private/GameProjectUtils.h b/Engine/Source/Editor/GameProjectGeneration/Private/GameProjectUtils.h
--- a/Engine/Source/Editor/GameProjectGeneration/Private/GameProjectUtils.h
+++ b/Engine/Source/Editor/GameProjectGeneration/Private/GameProjectUtils.h
@@ -21,6 +21,95 @@ public:
 		Classes,
 	};
 
+	/** Information used when creating a new class via AddCodeToProject */
+	struct FNewClassInfo
+	{
+		/** The type of class we want to create */
+		enum class EClassType : uint8
+		{
+			/** The new class is using a UObject as a base, consult BaseClass for the type */
+			UObject,
+
+			/** The new class should be an empty standard C++ class */
+			EmptyCpp,
+
+			/** The new class should be a Slate widget, deriving from SCompoundWidget */
+			SlateWidget,
+
+			/** The new class should be a Slate widget style, deriving from FSlateWidgetStyle, along with its associated UObject wrapper class */
+			SlateWidgetStyle,
+		};
+
+		/** Default constructor; must produce an object which fails the IsSet check */
+		FNewClassInfo()
+			: ClassType(EClassType::UObject)
+			, BaseClass(nullptr)
+		{
+		}
+
+		/** Convenience constructor so you can construct from a EClassType */
+		explicit FNewClassInfo(const EClassType InClassType)
+			: ClassType(InClassType)
+			, BaseClass(nullptr)
+		{
+		}
+
+		/** Convenience constructor so you can construct from a UClass */
+		explicit FNewClassInfo(const UClass* InBaseClass)
+			: ClassType(EClassType::UObject)
+			, BaseClass(InBaseClass)
+		{
+		}
+
+		/** Check to see if this struct is set to something that could be used to create a new class */
+		bool IsSet() const
+		{
+			return ClassType != EClassType::UObject || BaseClass;
+		}
+
+		/** Get the "friendly" class name to use in the UI */
+		FString GetClassName() const;
+
+		/** Get the class description to use in the UI */
+		FString GetClassDescription() const;
+
+		/** Get the class icon to use in the UI */
+		const FSlateBrush* GetClassIcon() const;
+
+		/** Get the C++ prefix used for this class type */
+		FString GetClassPrefixCPP() const;
+
+		/** Get the C++ class name; this may or may not be prefixed, but will always produce a valid C++ name via GetClassPrefix() + GetClassName() */
+		FString GetClassNameCPP() const;
+
+		/** Some classes may apply a particular suffix; this function returns the class name with those suffixes removed */
+		FString GetCleanClassName(const FString& ClassName) const;
+
+		/** Some classes may apply a particular suffix; this function returns the class name that will ultimately be used should that happen */
+		FString GetFinalClassName(const FString& ClassName) const;
+
+		/** Get the path needed to include this class into another file */
+		bool GetIncludePath(FString& OutIncludePath) const;
+
+		/** Given a class name, generate the header file (.h) that should be used for this class */
+		FString GetHeaderFilename(const FString& ClassName) const;
+
+		/** Given a class name, generate the source file (.cpp) that should be used for this class */
+		FString GetSourceFilename(const FString& ClassName) const;
+
+		/** Get the generation template filename to used based on the current class type */
+		FString GetHeaderTemplateFilename() const;
+
+		/** Get the generation template filename to used based on the current class type */
+		FString GetSourceTemplateFilename() const;
+
+		/** The type of class we want to create */
+		EClassType ClassType;
+
+		/** Base class information; if the ClassType is UObject */
+		const UClass* BaseClass;
+	};
+
 	/** Returns true if the project filename is properly formed and does not conflict with another project */
 	static bool IsValidProjectFileForCreation(const FString& ProjectFile, FText& OutFailReason);
 
@@ -52,7 +141,7 @@ public:
 	static bool IsValidClassNameForCreation(const FString& NewClassName, FText& OutFailReason);
 
 	/** Adds new source code to the project. When returning true, OutSyncFileAndLineNumber will be the the preferred target file to sync in the users code editing IDE, formatted for use with GenericApplication::GotoLineInSource */
-	static bool AddCodeToProject(const FString& NewClassName, const FString& NewClassPath, const UClass* ParentClass, FString& OutHeaderFilePath, FString& OutCppFilePath, FText& OutFailReason);
+	static bool AddCodeToProject(const FString& NewClassName, const FString& NewClassPath, const FNewClassInfo ParentClassInfo, FString& OutHeaderFilePath, FString& OutCppFilePath, FText& OutFailReason);
 
 	/** Loads a template project definitions object from the TemplateDefs.ini file in the specified project */
 	static UTemplateProjectDefs* LoadTemplateDefs(const FString& ProjectDirectory);
@@ -195,10 +284,10 @@ private:
 	static FString MakeIncludeList(const TArray<FString>& InList);
 
 	/** Generates a header file for a UObject class. OutSyncLocation is a string representing the preferred cursor sync location for this file after creation. */
-	static bool GenerateClassHeaderFile(const FString& NewHeaderFileName, const UClass* BaseClass, const TArray<FString>& ClassSpecifierList, const FString& ClassProperties, const FString& ClassFunctionDeclarations, FString& OutSyncLocation, FText& OutFailReason);
+	static bool GenerateClassHeaderFile(const FString& NewHeaderFileName, const FString UnPrefixedClassName, const FNewClassInfo ParentClassInfo, const TArray<FString>& ClassSpecifierList, const FString& ClassProperties, const FString& ClassFunctionDeclarations, FString& OutSyncLocation, FText& OutFailReason);
 
 	/** Generates a cpp file for a UObject class */
-	static bool GenerateClassCPPFile(const FString& NewCPPFileName, const FString& ModuleName, const FString& PrefixedClassName, const TArray<FString>& AdditionalIncludes, const TArray<FString>& PropertyOverrides, const FString& AdditionalMemberDefinitions, FText& OutFailReason);
+	static bool GenerateClassCPPFile(const FString& NewCPPFileName, const FString UnPrefixedClassName, const FNewClassInfo ParentClassInfo, const TArray<FString>& AdditionalIncludes, const TArray<FString>& PropertyOverrides, const FString& AdditionalMemberDefinitions, FText& OutFailReason);
 
 	/** Generates a Build.cs file for a game module */
 	static bool GenerateGameModuleBuildFile(const FString& NewBuildFileName, const FString& ModuleName, const TArray<FString>& PublicDependencyModuleNames, const TArray<FString>& PrivateDependencyModuleNames, FText& OutFailReason);
@@ -253,7 +342,7 @@ private:
 	static bool ProjectHasCodeFiles();
 
 	/** Internal handler for AddCodeToProject*/
-	static bool AddCodeToProject_Internal(const FString& NewClassName, const FString& NewClassPath, const UClass* ParentClass, FString& OutHeaderFilePath, FString& OutCppFilePath, FText& OutFailReason);
+	static bool AddCodeToProject_Internal(const FString& NewClassName, const FString& NewClassPath, const FNewClassInfo ParentClassInfo, FString& OutHeaderFilePath, FString& OutCppFilePath, FText& OutFailReason);
 
 	/** Handler for the user confirming they've read the name legnth warning */
 	static void OnWarningReasonOk();
diff --git a/Engine/Source/Editor/GameProjectGeneration/Private/SNewClassDialog.cpp b/Engine/Source/Editor/GameProjectGeneration/Private/SNewClassDialog.cpp
--- a/Engine/Source/Editor/GameProjectGeneration/Private/SNewClassDialog.cpp
+++ b/Engine/Source/Editor/GameProjectGeneration/Private/SNewClassDialog.cpp
@@ -3,7 +3,6 @@
 
 #include "GameProjectGenerationPrivatePCH.h"
 #include "SourceCodeNavigation.h"
-#include "ClassIconFinder.h"
 #include "ClassViewerModule.h"
 #include "ClassViewerFilter.h"
 #include "Editor/ClassViewer/Private/SClassViewer.h"
@@ -13,10 +12,10 @@
 
 struct FParentClassItem
 {
-	TWeakObjectPtr<UClass> ParentClass;
+	GameProjectUtils::FNewClassInfo ParentClassInfo;
 
-	FParentClassItem(const TWeakObjectPtr<UClass>& InParentClass)
-		: ParentClass(InParentClass)
+	FParentClassItem(const GameProjectUtils::FNewClassInfo& InParentClassInfo)
+		: ParentClassInfo(InParentClassInfo)
 	{}
 };
 
@@ -60,7 +59,7 @@ void SNewClassDialog::Construct( const FArguments& InArgs )
 	NewClassPath = GameProjectUtils::GetSourceRootPath(true/*bIncludeModuleName*/);
 	ClassLocation = GameProjectUtils::EClassLocation::UserDefined; // the first call to UpdateInputValidity will set this correctly based on NewClassPath
 
-	ParentClass = InArgs._Class;
+	ParentClassInfo = GameProjectUtils::FNewClassInfo(InArgs._Class);
 
 	bShowFullClassTree = false;
 
@@ -106,7 +105,7 @@ void SNewClassDialog::Construct( const FArguments& InArgs )
 				.FinishButtonToolTip ( LOCTEXT("FinishButtonToolTip", "Creates the code files to add your new class.").ToString() )
 				.OnCanceled(this, &SNewClassDialog::CancelClicked)
 				.OnFinished(this, &SNewClassDialog::FinishClicked)
-				.InitialPageIndex(ParentClass.IsValid() ? 1 : 0)
+				.InitialPageIndex(ParentClassInfo.IsSet() ? 1 : 0)
 				
 				// Choose parent class
 				+SWizard::Page()
@@ -557,25 +556,14 @@ TSharedRef<ITableRow> SNewClassDialog::MakeParentClassListViewWidget(TSharedPtr<
 		return SNew( STableRow<TSharedPtr<FParentClassItem>>, OwnerTable );
 	}
 
-	UClass* Class = ParentClassItem->ParentClass.Get();
-	if ( !Class )
+	if ( !ParentClassItem->ParentClassInfo.IsSet() )
 	{
 		return SNew( STableRow<TSharedPtr<FParentClassItem>>, OwnerTable );
 	}
 
-	const FString ClassName = FName::NameToDisplayString( Class->GetName(), false );
-	FString ClassDescription = Class->GetToolTipText().ToString();
-	int32 NewLineIndex = 0;
-	if (ClassDescription.FindChar('.', NewLineIndex))
-	{
-		// Only show the first sentence so as not to clutter up the UI with a detailed description of implementation details
-		ClassDescription = ClassDescription.Left(NewLineIndex + 1);
-	}
-
-	// Strip out any new-lines in the description
-	ClassDescription = ClassDescription.Replace(TEXT("\n"), TEXT(" "));
-
-	const FSlateBrush* ClassBrush = FClassIconFinder::FindIconForClass(Class);
+	const FString ClassName = ParentClassItem->ParentClassInfo.GetClassName();
+	const FString ClassDescription = ParentClassItem->ParentClassInfo.GetClassDescription();
+	const FSlateBrush* const ClassBrush = ParentClassItem->ParentClassInfo.GetClassIcon();
 
 	const int32 ItemHeight = 64;
 	const int32 DescriptionIndent = 32;
@@ -626,15 +614,7 @@ TSharedRef<ITableRow> SNewClassDialog::MakeParentClassListViewWidget(TSharedPtr<
 
 FString SNewClassDialog::GetSelectedParentClassName() const
 {
-	const UClass* SelectedParentClass = GetSelectedParentClass();
-	if (SelectedParentClass)
-	{
-		return SelectedParentClass->GetName();
-	}
-	else
-	{
-		return TEXT("");
-	}
+	return ParentClassInfo.IsSet() ? ParentClassInfo.GetClassName() : TEXT("");
 }
 
 void SNewClassDialog::OnParentClassItemDoubleClicked( TSharedPtr<FParentClassItem> TemplateItem )
@@ -652,18 +632,18 @@ void SNewClassDialog::OnClassSelected(TSharedPtr<FParentClassItem> Item, ESelect
 	if ( Item.IsValid() )
 	{
 		ClassViewer->ClearSelection();
-		ParentClass = Item->ParentClass;
+		ParentClassInfo = Item->ParentClassInfo;
 	}
 	else
 	{
-		ParentClass = NULL;
+		ParentClassInfo = GameProjectUtils::FNewClassInfo();
 	}
 }
 
 void SNewClassDialog::OnAdvancedClassSelected(UClass* Class)
 {
 	ParentClassListView->ClearSelection();
-	ParentClass = Class;
+	ParentClassInfo = GameProjectUtils::FNewClassInfo(Class);
 }
 
 ESlateCheckBoxState::Type SNewClassDialog::IsFullClassTreeChecked() const
@@ -724,24 +704,31 @@ FString SNewClassDialog::GetGlobalErrorLabelText() const
 void SNewClassDialog::OnNamePageEntered()
 {
 	// Set the default class name based on the selected parent class, eg MyActor
-	FString PotentialNewClassName = "My";
-	PotentialNewClassName += GetSelectedParentClassName();
+	const FString ParentClassName = ParentClassInfo.GetClassNameCPP();
+	const FString PotentialNewClassName = FString::Printf(TEXT("My%s"), ParentClassName.IsEmpty() ? TEXT("Class") : *ParentClassName);
 
 	// Only set the default if the user hasn't changed the class name from the previous default
 	if(LastAutoGeneratedClassName.IsEmpty() || NewClassName == LastAutoGeneratedClassName)
 	{
 		NewClassName = PotentialNewClassName;
 		LastAutoGeneratedClassName = PotentialNewClassName;
-		UpdateInputValidity();
 	}
 
+	UpdateInputValidity();
+
 	// Steal keyboard focus to accelerate name entering
 	FSlateApplication::Get().SetKeyboardFocus(ClassNameEditBox, EKeyboardFocusCause::SetDirectly);
 }
 
 FString SNewClassDialog::GetNameClassTitle() const
 {
-	return FText::Format( LOCTEXT( "NameClassTitle", "Name Your New {0}" ), FText::FromString(GetSelectedParentClassName()) ).ToString();
+	FString ParentClassName = GetSelectedParentClassName();
+	if(ParentClassName.IsEmpty() || ParentClassName == "None")
+	{
+		ParentClassName = TEXT("Class");
+	}
+
+	return FText::Format( LOCTEXT( "NameClassTitle", "Name Your New {0}" ), FText::FromString(ParentClassName) ).ToString();
 }
 
 FText SNewClassDialog::OnGetClassNameText() const
@@ -783,7 +770,7 @@ void SNewClassDialog::CancelClicked()
 
 bool SNewClassDialog::CanFinish() const
 {
-	return bLastInputValidityCheckSuccessful && GetSelectedParentClass() != NULL && FSourceCodeNavigation::IsCompilerAvailable();
+	return bLastInputValidityCheckSuccessful && ParentClassInfo.IsSet() && FSourceCodeNavigation::IsCompilerAvailable();
 }
 
 void SNewClassDialog::FinishClicked()
@@ -794,7 +781,7 @@ void SNewClassDialog::FinishClicked()
 	FString CppFilePath;
 
 	FText FailReason;
-	if ( GameProjectUtils::AddCodeToProject(NewClassName, NewClassPath, GetSelectedParentClass(), HeaderFilePath, CppFilePath, FailReason) )
+	if ( GameProjectUtils::AddCodeToProject(NewClassName, NewClassPath, ParentClassInfo, HeaderFilePath, CppFilePath, FailReason) )
 	{
 		// Prevent periodic validity checks. This is to prevent a brief error message about the class already existing while you are exiting.
 		bPreventPeriodicValidityChecksUntilNextChange = true;
@@ -959,8 +946,8 @@ void SNewClassDialog::UpdateInputValidity()
 	// Validate the path first since this has the side effect of updating the UI
 	FString ModuleName;
 	bLastInputValidityCheckSuccessful = GameProjectUtils::CalculateSourcePaths(NewClassPath, ModuleName, CalculatedClassHeaderName, CalculatedClassSourceName, &LastInputValidityErrorText);
-	CalculatedClassHeaderName /= NewClassName + ".h";
-	CalculatedClassSourceName /= NewClassName + ".cpp";
+	CalculatedClassHeaderName /= ParentClassInfo.GetHeaderFilename(NewClassName);
+	CalculatedClassSourceName /= ParentClassInfo.GetSourceFilename(NewClassName);
 
 	// If the source paths check as succeeded, check to see if we're using a Public/Private class
 	if(bLastInputValidityCheckSuccessful)
@@ -990,26 +977,33 @@ void SNewClassDialog::UpdateInputValidity()
 	bPreventPeriodicValidityChecksUntilNextChange = false;
 }
 
-const UClass* SNewClassDialog::GetSelectedParentClass() const
+const GameProjectUtils::FNewClassInfo& SNewClassDialog::GetSelectedParentClassInfo() const
 {
-	return ParentClass.Get();
+	return ParentClassInfo;
 }
 
 void SNewClassDialog::SetupParentClassItems()
 {
-	TArray<UClass*> FeaturedClasses;
+	TArray<GameProjectUtils::FNewClassInfo> FeaturedClasses;
 	
+	// Add an empty class
+	FeaturedClasses.Add(GameProjectUtils::FNewClassInfo(GameProjectUtils::FNewClassInfo::EClassType::EmptyCpp));
+
 	// @todo make this ini configurable
-	FeaturedClasses.Add(ACharacter::StaticClass());
-	FeaturedClasses.Add(APawn::StaticClass());
-	FeaturedClasses.Add(AActor::StaticClass());
-	FeaturedClasses.Add(APlayerCameraManager::StaticClass());
-	FeaturedClasses.Add(APlayerController::StaticClass());
-	FeaturedClasses.Add(AGameMode::StaticClass());
-	FeaturedClasses.Add(AWorldSettings::StaticClass());
-	FeaturedClasses.Add(AHUD::StaticClass());
-	FeaturedClasses.Add(APlayerState::StaticClass());
-	FeaturedClasses.Add(AGameState::StaticClass());
+	FeaturedClasses.Add(GameProjectUtils::FNewClassInfo(ACharacter::StaticClass()));
+	FeaturedClasses.Add(GameProjectUtils::FNewClassInfo(APawn::StaticClass()));
+	FeaturedClasses.Add(GameProjectUtils::FNewClassInfo(AActor::StaticClass()));
+	FeaturedClasses.Add(GameProjectUtils::FNewClassInfo(APlayerCameraManager::StaticClass()));
+	FeaturedClasses.Add(GameProjectUtils::FNewClassInfo(APlayerController::StaticClass()));
+	FeaturedClasses.Add(GameProjectUtils::FNewClassInfo(AGameMode::StaticClass()));
+	FeaturedClasses.Add(GameProjectUtils::FNewClassInfo(AWorldSettings::StaticClass()));
+	FeaturedClasses.Add(GameProjectUtils::FNewClassInfo(AHUD::StaticClass()));
+	FeaturedClasses.Add(GameProjectUtils::FNewClassInfo(APlayerState::StaticClass()));
+	FeaturedClasses.Add(GameProjectUtils::FNewClassInfo(AGameState::StaticClass()));
+	
+	// Add the extra non-UObject classes
+	FeaturedClasses.Add(GameProjectUtils::FNewClassInfo(GameProjectUtils::FNewClassInfo::EClassType::SlateWidget));
+	FeaturedClasses.Add(GameProjectUtils::FNewClassInfo(GameProjectUtils::FNewClassInfo::EClassType::SlateWidgetStyle));
 
 	for ( auto ClassIt = FeaturedClasses.CreateConstIterator(); ClassIt; ++ClassIt )
 	{
diff --git a/Engine/Source/Editor/GameProjectGeneration/Private/SNewClassDialog.h b/Engine/Source/Editor/GameProjectGeneration/Private/SNewClassDialog.h
--- a/Engine/Source/Editor/GameProjectGeneration/Private/SNewClassDialog.h
+++ b/Engine/Source/Editor/GameProjectGeneration/Private/SNewClassDialog.h
@@ -128,7 +128,7 @@ private:
 	void UpdateInputValidity();
 
 	/** Gets the currently selected parent class */
-	const UClass* GetSelectedParentClass() const;
+	const GameProjectUtils::FNewClassInfo& GetSelectedParentClassInfo() const;
 
 	/** Adds parent classes to the ParentClassListView source */
 	void SetupParentClassItems();
@@ -167,7 +167,7 @@ private:
 	FString LastAutoGeneratedClassName;
 
 	/** The selected parent class */
-	TWeakObjectPtr<UClass> ParentClass;
+	GameProjectUtils::FNewClassInfo ParentClassInfo;
 
 	/** If true, the full class tree will be shown in the parent class selection */
 	bool bShowFullClassTree;
