Fix infinite loop when duplicating a project if the last space is not followed by a version number.

diff --git a/Engine/Source/Editor/GameProjectGeneration/Private/GameProjectUtils.cpp b/Engine/Source/Editor/GameProjectGeneration/Private/GameProjectUtils.cpp
--- a/Engine/Source/Editor/GameProjectGeneration/Private/GameProjectUtils.cpp
+++ b/Engine/Source/Editor/GameProjectGeneration/Private/GameProjectUtils.cpp
@@ -1887,18 +1887,23 @@ bool GameProjectUtils::DuplicateProjectForUpgrade( const FString& InProjectFile,
 	for(int32 LastSpace; NewDirectoryName.FindLastChar(' ', LastSpace); )
 	{
 		const TCHAR *End = *NewDirectoryName + LastSpace + 1;
-		if(*End == '(')
+		if(*End != '(')
+		{
+			break;
+		}
+
+		End++;
+		while(FChar::IsDigit(*End) || *End == '.')
 		{
 			End++;
-			while(FChar::IsDigit(*End) || *End == '.')
-			{
-				End++;
-			}
-			if(*End == ')' && *(End + 1) == 0)
-			{
-				NewDirectoryName = NewDirectoryName.Left(LastSpace).TrimTrailing();
-			}
 		}
+
+		if(*End != ')' || *(End + 1) != 0)
+		{
+			break;
+		}
+
+		NewDirectoryName = NewDirectoryName.Left(LastSpace).TrimTrailing();
 	}
 
 	// Append the new version number
