Removing calls to clean a project when switching versions. Paranoid that it might do the wrong thing more often than the right thing after seeing a bug in the engine association code yesterday, since UBT running successfully has so many preconditions.

diff --git a/Engine/Source/Editor/GameProjectGeneration/Private/SProjectBrowser.cpp b/Engine/Source/Editor/GameProjectGeneration/Private/SProjectBrowser.cpp
--- a/Engine/Source/Editor/GameProjectGeneration/Private/SProjectBrowser.cpp
+++ b/Engine/Source/Editor/GameProjectGeneration/Private/SProjectBrowser.cpp
@@ -898,7 +898,7 @@ bool SProjectBrowser::OpenProject( const FString& InProjectFile )
 		// Ask whether to copy the project or upgrade in-place
 		TArray<FText> Buttons;
 		int32 OpenCopyButton = Buttons.Add(LOCTEXT("ProjectUpgradeCopy", "Create a copy"));
-		int32 InPlaceButton = Buttons.Add(LOCTEXT("ProjectUpgradeInPlace", "Convert in-place"));
+		Buttons.Add(LOCTEXT("ProjectUpgradeInPlace", "Convert in-place"));
 		int32 CancelButton = Buttons.Add(LOCTEXT("ProjectUpgradeCancel", "Cancel"));
 
 		// Show the dialog
@@ -913,16 +913,7 @@ bool SProjectBrowser::OpenProject( const FString& InProjectFile )
 			}
 			ProjectFile = NewProjectFile;
 		}
-		else if(Selection == InPlaceButton)
-		{
-			FString ProjectDir = FPaths::GetPath(ProjectFile);
-			if(!FDesktopPlatformModule::Get()->CleanGameProject(ProjectDir, GWarn))
-			{
-				FMessageDialog::Open( EAppMsgType::Ok, LOCTEXT("ConvertProjectCleanFailed", "Couldn't clean project binary and intermediate folders for upgrade. Check that the project folders are writable."));
-				return false;
-			}
-		}
-		else 
+		else if(Selection == CancelButton)
 		{
 			return false;
 		}
diff --git a/Engine/Source/Programs/Mac/UE4EditorServices/Private/Mac/UE4EditorServicesAppDelegate.cpp b/Engine/Source/Programs/Mac/UE4EditorServices/Private/Mac/UE4EditorServicesAppDelegate.cpp
--- a/Engine/Source/Programs/Mac/UE4EditorServices/Private/Mac/UE4EditorServicesAppDelegate.cpp
+++ b/Engine/Source/Programs/Mac/UE4EditorServices/Private/Mac/UE4EditorServicesAppDelegate.cpp
@@ -344,14 +344,6 @@
 						FPlatformMisc::MessageBoxExt(EAppMsgType::Ok, TEXT("Couldn't set association for project. Check the file is writeable."), TEXT("Error"));
 						return;
 					}
-
-					// Clean all the build products
-					FFeedbackContext *Warn = FDesktopPlatformModule::Get()->GetNativeFeedbackContext();
-					if (!FDesktopPlatformModule::Get()->CleanGameProject(FPaths::GetPath([FileURL path]), Warn))
-					{
-						FPlatformMisc::MessageBoxExt(EAppMsgType::Ok, TEXT("Couldn't clean project build products."), TEXT("Error"));
-						return;
-					}
 					break;
 				}
 			}
diff --git a/Engine/Source/Programs/UnrealVersionSelector/Private/UnrealVersionSelector.cpp b/Engine/Source/Programs/UnrealVersionSelector/Private/UnrealVersionSelector.cpp
--- a/Engine/Source/Programs/UnrealVersionSelector/Private/UnrealVersionSelector.cpp
+++ b/Engine/Source/Programs/UnrealVersionSelector/Private/UnrealVersionSelector.cpp
@@ -82,14 +82,6 @@ bool SwitchVersion(const FString& ProjectFileName)
 		return true;
 	}
 
-	// Otherwise clean all the build products
-	FFeedbackContext *Warn = FDesktopPlatformModule::Get()->GetNativeFeedbackContext();
-	if (!FDesktopPlatformModule::Get()->CleanGameProject(FPaths::GetPath(ProjectFileName), Warn))
-	{
-		FPlatformMisc::MessageBoxExt(EAppMsgType::Ok, TEXT("Couldn't clean project build products."), TEXT("Error"));
-		return false;
-	}
-
 	// Generate project files
 	return GenerateProjectFiles(ProjectFileName);
 }
