Fixed the code generated by the New Project wizard not compiling

TTP# 332794 - EDITOR: Usability improvements for the New Class Wizard

GetClassLocation was failing because it was doing the validation based upon the current project, and not the project that was being created. I've changed the validation logic and creation functions to use a module context that must be provided; this allows the functions used to generate code for a new project to override the module validation information.

diff --git a/Engine/Source/Editor/GameProjectGeneration/Private/GameProjectUtils.cpp b/Engine/Source/Editor/GameProjectGeneration/Private/GameProjectUtils.cpp
--- a/Engine/Source/Editor/GameProjectGeneration/Private/GameProjectUtils.cpp
+++ b/Engine/Source/Editor/GameProjectGeneration/Private/GameProjectUtils.cpp
@@ -1609,13 +1609,18 @@ bool GameProjectUtils::GenerateGameFrameworkSourceCode(const FString& NewProject
 {
 	const FString GameModulePath = NewProjectSourcePath / NewProjectName;
 
+	// Used to override the code generation validation since the module we're creating isn't the same as the project we currently have loaded
+	FModuleContextInfo NewModuleInfo;
+	NewModuleInfo.ModuleSourcePath = NewProjectSourcePath;
+	NewModuleInfo.ModuleName = NewProjectName;
+
 	// MyGamePlayerController.h
 	{
 		const UClass* BaseClass = APlayerController::StaticClass();
 		const FString NewClassName = NewProjectName + BaseClass->GetName();
 		const FString NewHeaderFilename = GameModulePath / NewClassName + TEXT(".h");
 		FString UnusedSyncLocation;
-		if ( GenerateClassHeaderFile(NewHeaderFilename, NewClassName, FNewClassInfo(BaseClass), TArray<FString>(), TEXT(""), TEXT(""), UnusedSyncLocation, OutFailReason) )
+		if ( GenerateClassHeaderFile(NewHeaderFilename, NewClassName, FNewClassInfo(BaseClass), TArray<FString>(), TEXT(""), TEXT(""), UnusedSyncLocation, NewModuleInfo, OutFailReason) )
 		{
 			OutCreatedFiles.Add(NewHeaderFilename);
 		}
@@ -1631,7 +1636,7 @@ bool GameProjectUtils::GenerateGameFrameworkSourceCode(const FString& NewProject
 		const FString NewClassName = NewProjectName + BaseClass->GetName();
 		const FString NewHeaderFilename = GameModulePath / NewClassName + TEXT(".h");
 		FString UnusedSyncLocation;
-		if ( GenerateClassHeaderFile(NewHeaderFilename, NewClassName, FNewClassInfo(BaseClass), TArray<FString>(), TEXT(""), TEXT(""), UnusedSyncLocation, OutFailReason) )
+		if ( GenerateClassHeaderFile(NewHeaderFilename, NewClassName, FNewClassInfo(BaseClass), TArray<FString>(), TEXT(""), TEXT(""), UnusedSyncLocation, NewModuleInfo, OutFailReason) )
 		{
 			OutCreatedFiles.Add(NewHeaderFilename);
 		}
@@ -1648,7 +1653,7 @@ bool GameProjectUtils::GenerateGameFrameworkSourceCode(const FString& NewProject
 		const FString NewClassName = NewProjectName + BaseClass->GetName();
 		const FString NewCPPFilename = GameModulePath / NewClassName + TEXT(".cpp");
 		PrefixedPlayerControllerClassName = FString(BaseClass->GetPrefixCPP()) + NewClassName;
-		if ( GenerateClassCPPFile(NewCPPFilename, NewClassName, FNewClassInfo(BaseClass), TArray<FString>(), TArray<FString>(), TEXT(""), OutFailReason) )
+		if ( GenerateClassCPPFile(NewCPPFilename, NewClassName, FNewClassInfo(BaseClass), TArray<FString>(), TArray<FString>(), TEXT(""), NewModuleInfo, OutFailReason) )
 		{
 			OutCreatedFiles.Add(NewCPPFilename);
 		}
@@ -1673,7 +1678,7 @@ bool GameProjectUtils::GenerateGameFrameworkSourceCode(const FString& NewProject
 		TArray<FString> AdditionalIncludes;
 		AdditionalIncludes.Add(PlayerControllerClassName);
 
-		if ( GenerateClassCPPFile(NewCPPFilename, NewClassName, FNewClassInfo(BaseClass), AdditionalIncludes, PropertyOverrides, TEXT(""), OutFailReason) )
+		if ( GenerateClassCPPFile(NewCPPFilename, NewClassName, FNewClassInfo(BaseClass), AdditionalIncludes, PropertyOverrides, TEXT(""), NewModuleInfo, OutFailReason) )
 		{
 			OutCreatedFiles.Add(NewCPPFilename);
 		}
@@ -1707,15 +1712,25 @@ bool GameProjectUtils::IsStarterContentAvailableForNewProjects()
 	return (StarterContentFiles.Num() > 0);
 }
 
-FString GameProjectUtils::GetSourceRootPath(const bool bIncludeModuleName)
+GameProjectUtils::FModuleContextInfo GameProjectUtils::GetCurrentModuleContextInfo()
+{
+	FModuleContextInfo ModuleInfo;
+
+	ModuleInfo.ModuleSourcePath = FPaths::GameSourceDir();
+	
+	// Assuming the game name is the same as the primary game module name
+	ModuleInfo.ModuleName = FApp::GetGameName();
+
+	return ModuleInfo;
+}
+
+FString GameProjectUtils::GetSourceRootPath(const bool bIncludeModuleName, const FModuleContextInfo& ModuleInfo)
 {
-	FString SourceDir = FPaths::GameSourceDir();
+	FString SourceDir = ModuleInfo.ModuleSourcePath;
 
 	if(bIncludeModuleName)
 	{
-		// Assuming the game name is the same as the primary game module name
-		const FString ModuleName = FApp::GetGameName();
-		SourceDir /= ModuleName;
+		SourceDir /= ModuleInfo.ModuleName;
 	}
 
 	SourceDir /= "";
@@ -1723,9 +1738,9 @@ FString GameProjectUtils::GetSourceRootPath(const bool bIncludeModuleName)
 	return FPaths::ConvertRelativePathToFull(SourceDir);
 }
 
-bool GameProjectUtils::IsValidSourcePath(const FString& InPath, const bool bIncludeModuleName, FText* const OutFailReason)
+bool GameProjectUtils::IsValidSourcePath(const FString& InPath, const bool bIncludeModuleName, const FModuleContextInfo& ModuleInfo, FText* const OutFailReason)
 {
-	FString RootPath = GetSourceRootPath(bIncludeModuleName);
+	FString RootPath = GetSourceRootPath(bIncludeModuleName, ModuleInfo);
 
 	// Only allow partial module name matches if we already have code; the first class added to a project *must* be for the game module
 	const bool bHasCodeFiles = GameProjectUtils::ProjectHasCodeFiles();
@@ -1749,7 +1764,7 @@ bool GameProjectUtils::IsValidSourcePath(const FString& InPath, const bool bIncl
 		if(OutFailReason)
 		{
 			FFormatNamedArguments Args;
-			Args.Add(TEXT("RootSourcePath"), FText::FromString(GetSourceRootPath(bIncludeModuleName)));
+			Args.Add(TEXT("RootSourcePath"), FText::FromString(GetSourceRootPath(bIncludeModuleName, ModuleInfo)));
 			const FText FormatString = (bIncludeModuleName)
 				? LOCTEXT("SourcePathInvalidModuleRoot", "All source code must exist within a valid module folder in the projects source path, e.g. {RootSourcePath}")
 				: LOCTEXT("SourcePathInvalidRoot", "All source code must exist within the projects source path: {RootSourcePath}");
@@ -1761,19 +1776,19 @@ bool GameProjectUtils::IsValidSourcePath(const FString& InPath, const bool bIncl
 	return true;
 }
 
-bool GameProjectUtils::CalculateSourcePaths(const FString& InPath, FString& OutModuleName, FString& OutHeaderPath, FString& OutSourcePath, FText* const OutFailReason)
+bool GameProjectUtils::CalculateSourcePaths(const FString& InPath, FString& OutModuleName, FString& OutHeaderPath, FString& OutSourcePath, const FModuleContextInfo& ModuleInfo, FText* const OutFailReason)
 {
 	const FString AbsoluteInPath = FPaths::ConvertRelativePathToFull(InPath) / ""; // Ensure trailing /
 	OutHeaderPath = AbsoluteInPath;
 	OutSourcePath = AbsoluteInPath;
 
 	EClassLocation ClassPathLocation = EClassLocation::UserDefined;
-	if(!GetClassLocation(InPath, OutModuleName, ClassPathLocation, OutFailReason))
+	if(!GetClassLocation(InPath, OutModuleName, ClassPathLocation, ModuleInfo, OutFailReason))
 	{
 		return false;
 	}
 
-	const FString BaseRootPath = GetSourceRootPath(false/*bIncludeModuleName*/);
+	const FString BaseRootPath = GetSourceRootPath(false/*bIncludeModuleName*/, ModuleInfo);
 	const FString RootPath = BaseRootPath / OutModuleName / ""; // Ensure trailing /
 	const FString PublicPath = RootPath / "Public" / "";		// Ensure trailing /
 	const FString PrivatePath = RootPath / "Private" / "";		// Ensure trailing /
@@ -1821,20 +1836,20 @@ bool GameProjectUtils::CalculateSourcePaths(const FString& InPath, FString& OutM
 	return !OutHeaderPath.IsEmpty() && !OutSourcePath.IsEmpty();
 }
 
-bool GameProjectUtils::GetClassLocation(const FString& InPath, FString& OutModuleName, EClassLocation& OutClassLocation, FText* const OutFailReason)
+bool GameProjectUtils::GetClassLocation(const FString& InPath, FString& OutModuleName, EClassLocation& OutClassLocation, const FModuleContextInfo& ModuleInfo, FText* const OutFailReason)
 {
 	const FString AbsoluteInPath = FPaths::ConvertRelativePathToFull(InPath) / ""; // Ensure trailing /
 	OutModuleName.Empty();
 	OutClassLocation = EClassLocation::UserDefined;
 
-	if(!IsValidSourcePath(InPath, true/*bIncludeModuleName*/, OutFailReason))
+	if(!IsValidSourcePath(InPath, true/*bIncludeModuleName*/, ModuleInfo, OutFailReason))
 	{
 		return false;
 	}
 
 	// We've validated that this path includes a partial match for our module (eg, MyModule, MyModuleEditor, MyModuleClient)
 	// so extract the actual name of the module from the path so that we can generate the internal folder names correctly
-	const FString BaseRootPath = GetSourceRootPath(false/*bIncludeModuleName*/);
+	const FString BaseRootPath = GetSourceRootPath(false/*bIncludeModuleName*/, ModuleInfo);
 	const int32 ModuleNameStartIndex = BaseRootPath.Len();
 	const int32 ModuleNameEndIndex = AbsoluteInPath.Find(TEXT("/"), ESearchCase::IgnoreCase, ESearchDir::FromStart, ModuleNameStartIndex);
 	check(ModuleNameEndIndex != INDEX_NONE); // this should never happen since AbsoluteInPath ends in a /, and we verified it started with BaseRootPath in IsValidSourcePath
@@ -2052,7 +2067,7 @@ FString GameProjectUtils::MakeIncludeList(const TArray<FString>& InList)
 	return ReturnString;
 }
 
-bool GameProjectUtils::GenerateClassHeaderFile(const FString& NewHeaderFileName, const FString UnPrefixedClassName, const FNewClassInfo ParentClassInfo, const TArray<FString>& ClassSpecifierList, const FString& ClassProperties, const FString& ClassFunctionDeclarations, FString& OutSyncLocation, FText& OutFailReason)
+bool GameProjectUtils::GenerateClassHeaderFile(const FString& NewHeaderFileName, const FString UnPrefixedClassName, const FNewClassInfo ParentClassInfo, const TArray<FString>& ClassSpecifierList, const FString& ClassProperties, const FString& ClassFunctionDeclarations, FString& OutSyncLocation, const FModuleContextInfo& ModuleInfo, FText& OutFailReason)
 {
 	FString Template;
 	if ( !ReadTemplateFile(ParentClassInfo.GetHeaderTemplateFilename(), Template, OutFailReason) )
@@ -2075,12 +2090,13 @@ bool GameProjectUtils::GenerateClassHeaderFile(const FString& NewHeaderFileName,
 	{
 		FString ModuleName;
 		EClassLocation ClassPathLocation = EClassLocation::UserDefined;
-		GetClassLocation(NewHeaderFileName, ModuleName, ClassPathLocation);
-
-		// If this class isn't Private, make sure and include the API macro so it can be linked within other modules
-		if ( ClassPathLocation != EClassLocation::Private )
+		if ( GetClassLocation(NewHeaderFileName, ModuleName, ClassPathLocation, ModuleInfo) )
 		{
-			ModuleAPIMacro = ModuleName.ToUpper() + "_API "; // include a trailing space for the template formatting
+			// If this class isn't Private, make sure and include the API macro so it can be linked within other modules
+			if ( ClassPathLocation != EClassLocation::Private )
+			{
+				ModuleAPIMacro = ModuleName.ToUpper() + "_API "; // include a trailing space for the template formatting
+			}
 		}
 	}
 
@@ -2122,7 +2138,7 @@ bool GameProjectUtils::GenerateClassHeaderFile(const FString& NewHeaderFileName,
 	return WriteOutputFile(NewHeaderFileName, FinalOutput, OutFailReason);
 }
 
-bool GameProjectUtils::GenerateClassCPPFile(const FString& NewCPPFileName, const FString UnPrefixedClassName, const FNewClassInfo ParentClassInfo, const TArray<FString>& AdditionalIncludes, const TArray<FString>& PropertyOverrides, const FString& AdditionalMemberDefinitions, FText& OutFailReason)
+bool GameProjectUtils::GenerateClassCPPFile(const FString& NewCPPFileName, const FString UnPrefixedClassName, const FNewClassInfo ParentClassInfo, const TArray<FString>& AdditionalIncludes, const TArray<FString>& PropertyOverrides, const FString& AdditionalMemberDefinitions, const FModuleContextInfo& ModuleInfo, FText& OutFailReason)
 {
 	FString Template;
 	if ( !ReadTemplateFile(ParentClassInfo.GetSourceTemplateFilename(), Template, OutFailReason) )
@@ -2136,7 +2152,10 @@ bool GameProjectUtils::GenerateClassCPPFile(const FString& NewCPPFileName, const
 
 	FString ModuleName;
 	EClassLocation ClassPathLocation = EClassLocation::UserDefined;
-	GetClassLocation(NewCPPFileName, ModuleName, ClassPathLocation);
+	if ( !GetClassLocation(NewCPPFileName, ModuleName, ClassPathLocation, ModuleInfo, &OutFailReason) )
+	{
+		return false;
+	}
 
 	FString AdditionalIncludesStr;
 	for (int32 IncludeIdx = 0; IncludeIdx < AdditionalIncludes.Num(); ++IncludeIdx)
@@ -2161,18 +2180,10 @@ bool GameProjectUtils::GenerateClassCPPFile(const FString& NewCPPFileName, const
 		PropertyOverridesStr += *PropertyOverrides[OverrideIdx];
 	}
 
-	// account for subfolders in new filename
-	FString ModuleHeaderPath = ModuleName;
-	int32 SubfolderCount = CalculateSubfolderCount(NewCPPFileName, ModuleName);
-	for (int32 i = 0; i < SubfolderCount; ++i)
-	{
-		ModuleHeaderPath = "../" + ModuleHeaderPath;
-	}
-
 	// Not all of these will exist in every class template
 	FString FinalOutput = Template.Replace(TEXT("%COPYRIGHT_LINE%"), *MakeCopyrightLine(), ESearchCase::CaseSensitive);
 	FinalOutput = FinalOutput.Replace(TEXT("%UNPREFIXED_CLASS_NAME%"), *UnPrefixedClassName, ESearchCase::CaseSensitive);
-	FinalOutput = FinalOutput.Replace(TEXT("%MODULE_NAME%"), *ModuleHeaderPath, ESearchCase::CaseSensitive);
+	FinalOutput = FinalOutput.Replace(TEXT("%MODULE_NAME%"), *ModuleName, ESearchCase::CaseSensitive);
 	FinalOutput = FinalOutput.Replace(TEXT("%PREFIXED_CLASS_NAME%"), *PrefixedClassName, ESearchCase::CaseSensitive);
 	FinalOutput = FinalOutput.Replace(TEXT("%PROPERTY_OVERRIDES%"), *PropertyOverridesStr, ESearchCase::CaseSensitive);
 	FinalOutput = FinalOutput.Replace(TEXT("%ADDITIONAL_MEMBER_DEFINITIONS%"), *AdditionalMemberDefinitions, ESearchCase::CaseSensitive);
@@ -2511,10 +2522,13 @@ bool GameProjectUtils::AddCodeToProject_Internal(const FString& NewClassName, co
 		return false;
 	}
 
+	// Get the context info for the current project
+	const FModuleContextInfo ModuleInfo = GetCurrentModuleContextInfo();
+
 	FString ModuleName;
 	FString NewHeaderPath;
 	FString NewCppPath;
-	if ( !CalculateSourcePaths(NewClassPath, ModuleName, NewHeaderPath, NewCppPath, &OutFailReason) )
+	if ( !CalculateSourcePaths(NewClassPath, ModuleName, NewHeaderPath, NewCppPath, ModuleInfo, &OutFailReason) )
 	{
 		return false;
 	}
@@ -2548,7 +2562,7 @@ bool GameProjectUtils::AddCodeToProject_Internal(const FString& NewClassName, co
 	FString SyncLocation;
 	const FString NewHeaderFilename = NewHeaderPath / ParentClassInfo.GetHeaderFilename(NewClassName);
 	{
-		if ( GenerateClassHeaderFile(NewHeaderFilename, CleanClassName, ParentClassInfo, TArray<FString>(), TEXT(""), TEXT(""), SyncLocation, OutFailReason) )
+		if ( GenerateClassHeaderFile(NewHeaderFilename, CleanClassName, ParentClassInfo, TArray<FString>(), TEXT(""), TEXT(""), SyncLocation, ModuleInfo, OutFailReason) )
 		{
 			CreatedFiles.Add(NewHeaderFilename);
 		}
@@ -2562,7 +2576,7 @@ bool GameProjectUtils::AddCodeToProject_Internal(const FString& NewClassName, co
 	// Class CPP file
 	const FString NewCppFilename = NewCppPath / ParentClassInfo.GetSourceFilename(NewClassName);
 	{
-		if ( GenerateClassCPPFile(NewCppFilename, CleanClassName, ParentClassInfo, TArray<FString>(), TArray<FString>(), TEXT(""), OutFailReason) )
+		if ( GenerateClassCPPFile(NewCppFilename, CleanClassName, ParentClassInfo, TArray<FString>(), TArray<FString>(), TEXT(""), ModuleInfo, OutFailReason) )
 		{
 			CreatedFiles.Add(NewCppFilename);
 		}
diff --git a/Engine/Source/Editor/GameProjectGeneration/Private/GameProjectUtils.h b/Engine/Source/Editor/GameProjectGeneration/Private/GameProjectUtils.h
--- a/Engine/Source/Editor/GameProjectGeneration/Private/GameProjectUtils.h
+++ b/Engine/Source/Editor/GameProjectGeneration/Private/GameProjectUtils.h
@@ -5,6 +5,16 @@
 class GameProjectUtils
 {
 public:
+	/** Context information used when validating that source code is being placed in the correct place for a given module */
+	struct FModuleContextInfo
+	{
+		/** Path to the Source folder of the module */
+		FString ModuleSourcePath;
+
+		/** Name of the module */
+		FString ModuleName;
+	};
+
 	/** Where is this class located within the Source folder? */
 	enum class EClassLocation : uint8
 	{
@@ -164,25 +174,32 @@ public:
 	/** Returns true if there are starter content files available for instancing into new projects. */
 	static bool IsStarterContentAvailableForNewProjects();
 
+	/**
+	 * Get the default module context information based on the current value of FApp
+	 */
+	static FModuleContextInfo GetCurrentModuleContextInfo();
+
 	/** 
 	 * Get the absolute root path under which all project source code must exist
 	 *
 	 * @param	bIncludeModuleName	Whether to include the module name in the root path?
+	 * @param	ModuleInfo			Information about the module being validated; use GetCurrentModuleContextInfo() if you want to validate the currently loaded project
 	 * 
 	 * @return	The root path. Will always be an absolute path ending with a /
 	 */
-	static FString GetSourceRootPath(const bool bIncludeModuleName);
+	static FString GetSourceRootPath(const bool bIncludeModuleName, const FModuleContextInfo& ModuleInfo);
 
 	/** 
 	 * Check to see if the given path is a valid place to put source code for this project (exists within the source root path) 
 	 *
 	 * @param	InPath				The path to check
 	 * @param	bIncludeModuleName	Whether to require the module name in the root path? (is really used as a prefix so will allow MyModule, MyModuleEditor, etc)
+	 * @param	ModuleInfo			Information about the module being validated; use GetCurrentModuleContextInfo() if you want to validate the currently loaded project
 	 * @param	OutFailReason		Optional parameter to fill with failure information
 	 * 
 	 * @return	true if the path is valid, false otherwise
 	 */
-	static bool IsValidSourcePath(const FString& InPath, const bool bIncludeModuleName, FText* const OutFailReason = nullptr);
+	static bool IsValidSourcePath(const FString& InPath, const bool bIncludeModuleName, const FModuleContextInfo& ModuleInfo, FText* const OutFailReason = nullptr);
 
 	/** 
 	 * Given the path provided, work out where generated .h and .cpp files would be placed
@@ -191,11 +208,12 @@ public:
 	 * @param	OutModuleName		The module name extracted from the path (the part after GetSourceRootPath(false))
 	 * @param	OutHeaderPath		The path where the .h file should be placed
 	 * @param	OutSourcePath		The path where the .cpp file should be placed
+	 * @param	ModuleInfo			Information about the module being validated; use GetCurrentModuleContextInfo() if you want to validate the currently loaded project
 	 * @param	OutFailReason		Optional parameter to fill with failure information
 	 * 
 	 * @return	false if the paths are invalid
 	 */
-	static bool CalculateSourcePaths(const FString& InPath, FString& OutModuleName, FString& OutHeaderPath, FString& OutSourcePath, FText* const OutFailReason = nullptr);
+	static bool CalculateSourcePaths(const FString& InPath, FString& OutModuleName, FString& OutHeaderPath, FString& OutSourcePath, const FModuleContextInfo& ModuleInfo, FText* const OutFailReason = nullptr);
 
 	/** 
 	 * Given the path provided, work out where it's located within the Source folder
@@ -203,11 +221,12 @@ public:
 	 * @param	InPath				The path to use a base
 	 * @param	OutModuleName		The module name extracted from the path (the part after GetSourceRootPath(false))
 	 * @param	OutClassLocation	The location within the Source folder
+	 * @param	ModuleInfo			Information about the module being validated; use GetCurrentModuleContextInfo() if you want to validate the currently loaded project
 	 * @param	OutFailReason		Optional parameter to fill with failure information
 	 * 
 	 * @return	false if the paths are invalid
 	 */
-	static bool GetClassLocation(const FString& InPath, FString& OutModuleName, EClassLocation& OutClassLocation, FText* const OutFailReason = nullptr);
+	static bool GetClassLocation(const FString& InPath, FString& OutModuleName, EClassLocation& OutClassLocation, const FModuleContextInfo& ModuleInfo, FText* const OutFailReason = nullptr);
 
 	/** Creates a copy of a project directory in order to upgrade it. */
 	static bool DuplicateProjectForUpgrade( const FString& InProjectFile, FString &OutNewProjectFile );
@@ -284,10 +303,10 @@ private:
 	static FString MakeIncludeList(const TArray<FString>& InList);
 
 	/** Generates a header file for a UObject class. OutSyncLocation is a string representing the preferred cursor sync location for this file after creation. */
-	static bool GenerateClassHeaderFile(const FString& NewHeaderFileName, const FString UnPrefixedClassName, const FNewClassInfo ParentClassInfo, const TArray<FString>& ClassSpecifierList, const FString& ClassProperties, const FString& ClassFunctionDeclarations, FString& OutSyncLocation, FText& OutFailReason);
+	static bool GenerateClassHeaderFile(const FString& NewHeaderFileName, const FString UnPrefixedClassName, const FNewClassInfo ParentClassInfo, const TArray<FString>& ClassSpecifierList, const FString& ClassProperties, const FString& ClassFunctionDeclarations, FString& OutSyncLocation, const FModuleContextInfo& ModuleInfo, FText& OutFailReason);
 
 	/** Generates a cpp file for a UObject class */
-	static bool GenerateClassCPPFile(const FString& NewCPPFileName, const FString UnPrefixedClassName, const FNewClassInfo ParentClassInfo, const TArray<FString>& AdditionalIncludes, const TArray<FString>& PropertyOverrides, const FString& AdditionalMemberDefinitions, FText& OutFailReason);
+	static bool GenerateClassCPPFile(const FString& NewCPPFileName, const FString UnPrefixedClassName, const FNewClassInfo ParentClassInfo, const TArray<FString>& AdditionalIncludes, const TArray<FString>& PropertyOverrides, const FString& AdditionalMemberDefinitions, const FModuleContextInfo& ModuleInfo, FText& OutFailReason);
 
 	/** Generates a Build.cs file for a game module */
 	static bool GenerateGameModuleBuildFile(const FString& NewBuildFileName, const FString& ModuleName, const TArray<FString>& PublicDependencyModuleNames, const TArray<FString>& PrivateDependencyModuleNames, FText& OutFailReason);
diff --git a/Engine/Source/Editor/GameProjectGeneration/Private/SNewClassDialog.cpp b/Engine/Source/Editor/GameProjectGeneration/Private/SNewClassDialog.cpp
--- a/Engine/Source/Editor/GameProjectGeneration/Private/SNewClassDialog.cpp
+++ b/Engine/Source/Editor/GameProjectGeneration/Private/SNewClassDialog.cpp
@@ -56,7 +56,9 @@ public:
 BEGIN_SLATE_FUNCTION_BUILD_OPTIMIZATION
 void SNewClassDialog::Construct( const FArguments& InArgs )
 {
-	NewClassPath = GameProjectUtils::GetSourceRootPath(true/*bIncludeModuleName*/);
+	ModuleInfo = GameProjectUtils::GetCurrentModuleContextInfo();
+
+	NewClassPath = GameProjectUtils::GetSourceRootPath(true/*bIncludeModuleName*/, ModuleInfo);
 	ClassLocation = GameProjectUtils::EClassLocation::UserDefined; // the first call to UpdateInputValidity will set this correctly based on NewClassPath
 
 	ParentClassInfo = GameProjectUtils::FNewClassInfo(InArgs._Class);
@@ -893,9 +895,9 @@ void SNewClassDialog::OnClassLocationChanged(ESlateCheckBoxState::Type InChecked
 
 		FString ModuleName;
 		GameProjectUtils::EClassLocation TmpClassLocation = GameProjectUtils::EClassLocation::UserDefined;
-		GameProjectUtils::GetClassLocation(AbsoluteClassPath, ModuleName, TmpClassLocation);
+		GameProjectUtils::GetClassLocation(AbsoluteClassPath, ModuleName, TmpClassLocation, ModuleInfo);
 
-		const FString BaseRootPath = GameProjectUtils::GetSourceRootPath(false/*bIncludeModuleName*/);
+		const FString BaseRootPath = GameProjectUtils::GetSourceRootPath(false/*bIncludeModuleName*/, ModuleInfo);
 		const FString RootPath = BaseRootPath / ModuleName / "";	// Ensure trailing /
 		const FString PublicPath = RootPath / "Public" / "";		// Ensure trailing /
 		const FString PrivatePath = RootPath / "Private" / "";		// Ensure trailing /
@@ -945,14 +947,14 @@ void SNewClassDialog::UpdateInputValidity()
 
 	// Validate the path first since this has the side effect of updating the UI
 	FString ModuleName;
-	bLastInputValidityCheckSuccessful = GameProjectUtils::CalculateSourcePaths(NewClassPath, ModuleName, CalculatedClassHeaderName, CalculatedClassSourceName, &LastInputValidityErrorText);
+	bLastInputValidityCheckSuccessful = GameProjectUtils::CalculateSourcePaths(NewClassPath, ModuleName, CalculatedClassHeaderName, CalculatedClassSourceName, ModuleInfo, &LastInputValidityErrorText);
 	CalculatedClassHeaderName /= ParentClassInfo.GetHeaderFilename(NewClassName);
 	CalculatedClassSourceName /= ParentClassInfo.GetSourceFilename(NewClassName);
 
 	// If the source paths check as succeeded, check to see if we're using a Public/Private class
 	if(bLastInputValidityCheckSuccessful)
 	{
-		GameProjectUtils::GetClassLocation(NewClassPath, ModuleName, ClassLocation);
+		GameProjectUtils::GetClassLocation(NewClassPath, ModuleName, ClassLocation, ModuleInfo);
 
 		// We only care about the Public and Private folders
 		if(ClassLocation != GameProjectUtils::EClassLocation::Public && ClassLocation != GameProjectUtils::EClassLocation::Private)
diff --git a/Engine/Source/Editor/GameProjectGeneration/Private/SNewClassDialog.h b/Engine/Source/Editor/GameProjectGeneration/Private/SNewClassDialog.h
--- a/Engine/Source/Editor/GameProjectGeneration/Private/SNewClassDialog.h
+++ b/Engine/Source/Editor/GameProjectGeneration/Private/SNewClassDialog.h
@@ -189,4 +189,7 @@ private:
 
 	/** Whether the class should be created as a Public or Private class */
 	GameProjectUtils::EClassLocation ClassLocation;
+
+	/** Information about the currently loaded module; used for class validation */
+	GameProjectUtils::FModuleContextInfo ModuleInfo;
 };
