Add a button to skip the conversion process for projects.

diff --git a/Engine/Source/Editor/GameProjectGeneration/Private/SProjectBrowser.cpp b/Engine/Source/Editor/GameProjectGeneration/Private/SProjectBrowser.cpp
--- a/Engine/Source/Editor/GameProjectGeneration/Private/SProjectBrowser.cpp
+++ b/Engine/Source/Editor/GameProjectGeneration/Private/SProjectBrowser.cpp
@@ -899,10 +899,11 @@ bool SProjectBrowser::OpenProject( const FString& InProjectFile )
 		TArray<FText> Buttons;
 		int32 OpenCopyButton = Buttons.Add(LOCTEXT("ProjectUpgradeCopy", "Create a copy"));
 		Buttons.Add(LOCTEXT("ProjectUpgradeInPlace", "Convert in-place"));
+		int32 SkipConversionButton = Buttons.Add(LOCTEXT("ProjectUpgradeSkipConversion", "Skip conversion"));
 		int32 CancelButton = Buttons.Add(LOCTEXT("ProjectUpgradeCancel", "Cancel"));
 
 		// Show the dialog
-		int32 Selection = SVerbChoiceDialog::ShowModal(LOCTEXT("ProjectConversionTitle", "Convert Project"), LOCTEXT("ConvertProjectPrompt", "This project was made with a different version of the editor. Converting to this version may be irreversible or lose data.\n\nIt is recommened that you create a copy before proceeding."), Buttons);
+		int32 Selection = SVerbChoiceDialog::ShowModal(LOCTEXT("ProjectConversionTitle", "Convert Project"), LOCTEXT("ConvertProjectPrompt", "This project was made with a different version of the editor. Opening with this version may be irreversible or lose data.\n\nIt is recommened that you create a copy before proceeding."), Buttons);
 		if(Selection == OpenCopyButton)
 		{
 			FString NewProjectFile;
@@ -918,68 +919,72 @@ bool SProjectBrowser::OpenProject( const FString& InProjectFile )
 			return false;
 		}
 
-		// If it's a code-based project, generate project files and open visual studio after an upgrade
-		if(ProjectStatus.bCodeBasedProject)
+		// Unless it was skipped, do the conversion
+		if(Selection != SkipConversionButton)
 		{
-			// Try to generate project files
-			FStringOutputDevice OutputLog;
-			OutputLog.SetAutoEmitLineTerminator(true);
-			GLog->AddOutputDevice(&OutputLog);
-			bool bHaveProjectFiles = FDesktopPlatformModule::Get()->GenerateProjectFiles(FPaths::RootDir(), ProjectFile, GWarn);
-			GLog->RemoveOutputDevice(&OutputLog);
-
-			// Display any errors
-			if(!bHaveProjectFiles)
+			// If it's a code-based project, generate project files and open visual studio after an upgrade
+			if(ProjectStatus.bCodeBasedProject)
 			{
-				FFormatNamedArguments Args;
-				Args.Add( TEXT("LogOutput"), FText::FromString(OutputLog) );
-				FMessageDialog::Open(EAppMsgType::Ok, FText::Format(LOCTEXT("CouldNotGenerateProjectFiles", "Project files could not be generated. Log output:\n\n{LogOutput}"), Args));
-				return false;
-			}
-
-			// Try to compile the project
-			OutputLog.Empty();
-			GLog->AddOutputDevice(&OutputLog);
-			bool bCompileSucceeded = FDesktopPlatformModule::Get()->CompileGameProject(FPaths::RootDir(), ProjectFile, GWarn);
-			GLog->RemoveOutputDevice(&OutputLog);
-
-			// Try to compile the modules
-			if(!bCompileSucceeded)
-			{
-				FText DevEnvName = FSourceCodeNavigation::GetSuggestedSourceCodeIDE( true );
-
-				TArray<FText> CompileFailedButtons;
-				int32 OpenIDEButton = CompileFailedButtons.Add(FText::Format(LOCTEXT("CompileFailedOpenIDE", "Open with {0}"), DevEnvName));
-				int32 ViewLogButton = CompileFailedButtons.Add(LOCTEXT("CompileFailedViewLog", "View build log"));
-				CompileFailedButtons.Add(LOCTEXT("CompileFailedCancel", "Cancel"));
-
-				int32 CompileFailedChoice = SVerbChoiceDialog::ShowModal(LOCTEXT("ProjectUpgradeTitle", "Project Conversion Failed"), FText::Format(LOCTEXT("ProjectUpgradeCompileFailed", "The project failed to compile with this version of the engine. Would you like to open the project in {0}?"), DevEnvName), CompileFailedButtons);
-				if(CompileFailedChoice == ViewLogButton)
+				// Try to generate project files
+				FStringOutputDevice OutputLog;
+				OutputLog.SetAutoEmitLineTerminator(true);
+				GLog->AddOutputDevice(&OutputLog);
+				bool bHaveProjectFiles = FDesktopPlatformModule::Get()->GenerateProjectFiles(FPaths::RootDir(), ProjectFile, GWarn);
+				GLog->RemoveOutputDevice(&OutputLog);
+
+				// Display any errors
+				if(!bHaveProjectFiles)
 				{
-					CompileFailedButtons.RemoveAt(ViewLogButton);
-					CompileFailedChoice = SVerbChoiceDialog::ShowModal(LOCTEXT("ProjectUpgradeTitle", "Project Conversion Failed"), FText::Format(LOCTEXT("ProjectUpgradeCompileFailed", "The project failed to compile with this version of the engine. Build output is as follows:\n\n{0}"), FText::FromString(OutputLog)), CompileFailedButtons);
+					FFormatNamedArguments Args;
+					Args.Add( TEXT("LogOutput"), FText::FromString(OutputLog) );
+					FMessageDialog::Open(EAppMsgType::Ok, FText::Format(LOCTEXT("CouldNotGenerateProjectFiles", "Project files could not be generated. Log output:\n\n{LogOutput}"), Args));
+					return false;
 				}
 
-				if(CompileFailedChoice == OpenIDEButton && !GameProjectUtils::OpenCodeIDE(ProjectFile, FailReason))
+				// Try to compile the project
+				OutputLog.Empty();
+				GLog->AddOutputDevice(&OutputLog);
+				bool bCompileSucceeded = FDesktopPlatformModule::Get()->CompileGameProject(FPaths::RootDir(), ProjectFile, GWarn);
+				GLog->RemoveOutputDevice(&OutputLog);
+
+				// Try to compile the modules
+				if(!bCompileSucceeded)
 				{
-					FMessageDialog::Open(EAppMsgType::Ok, FailReason);
+					FText DevEnvName = FSourceCodeNavigation::GetSuggestedSourceCodeIDE( true );
+
+					TArray<FText> CompileFailedButtons;
+					int32 OpenIDEButton = CompileFailedButtons.Add(FText::Format(LOCTEXT("CompileFailedOpenIDE", "Open with {0}"), DevEnvName));
+					int32 ViewLogButton = CompileFailedButtons.Add(LOCTEXT("CompileFailedViewLog", "View build log"));
+					CompileFailedButtons.Add(LOCTEXT("CompileFailedCancel", "Cancel"));
+
+					int32 CompileFailedChoice = SVerbChoiceDialog::ShowModal(LOCTEXT("ProjectUpgradeTitle", "Project Conversion Failed"), FText::Format(LOCTEXT("ProjectUpgradeCompileFailed", "The project failed to compile with this version of the engine. Would you like to open the project in {0}?"), DevEnvName), CompileFailedButtons);
+					if(CompileFailedChoice == ViewLogButton)
+					{
+						CompileFailedButtons.RemoveAt(ViewLogButton);
+						CompileFailedChoice = SVerbChoiceDialog::ShowModal(LOCTEXT("ProjectUpgradeTitle", "Project Conversion Failed"), FText::Format(LOCTEXT("ProjectUpgradeCompileFailed", "The project failed to compile with this version of the engine. Build output is as follows:\n\n{0}"), FText::FromString(OutputLog)), CompileFailedButtons);
+					}
+
+					if(CompileFailedChoice == OpenIDEButton && !GameProjectUtils::OpenCodeIDE(ProjectFile, FailReason))
+					{
+						FMessageDialog::Open(EAppMsgType::Ok, FailReason);
+					}
+
+					return false;
 				}
+			}
 
+			// Update the game project to the latest version. This will prompt to check-out as necessary. We don't need to write the engine identifier yet, because it won't use the right .uprojectdirs logic.
+			if(!GameProjectUtils::UpdateGameProject(TEXT("")))
+			{
 				return false;
 			}
-		}
-
-		// Update the game project to the latest version. This will prompt to check-out as necessary. We don't need to write the engine identifier yet, because it won't use the right .uprojectdirs logic.
-		if(!GameProjectUtils::UpdateGameProject(TEXT("")))
-		{
-			return false;
-		}
 
-		// Write the engine association
-		if(!FDesktopPlatformModule::Get()->SetEngineIdentifierForProject(ProjectFile, CurrentIdentifier))
-		{
-			FMessageDialog::Open(EAppMsgType::Ok, LOCTEXT("ProjectUpgradeFailure", "Project file could not be updated to latest version. Check that the file is writable."));
-			return false;
+			// Write the engine association
+			if(!FDesktopPlatformModule::Get()->SetEngineIdentifierForProject(ProjectFile, CurrentIdentifier))
+			{
+				FMessageDialog::Open(EAppMsgType::Ok, LOCTEXT("ProjectUpgradeFailure", "Project file could not be updated to latest version. Check that the file is writable."));
+				return false;
+			}
 		}
 	}
 
