Improve wording for project upgrade dialog depending on whether you're upgrading a code or content project. Also fix alignment of buttons in dialog.

diff --git a/Engine/Source/Editor/GameProjectGeneration/Private/SProjectBrowser.cpp b/Engine/Source/Editor/GameProjectGeneration/Private/SProjectBrowser.cpp
--- a/Engine/Source/Editor/GameProjectGeneration/Private/SProjectBrowser.cpp
+++ b/Engine/Source/Editor/GameProjectGeneration/Private/SProjectBrowser.cpp
@@ -895,15 +895,29 @@ bool SProjectBrowser::OpenProject( const FString& InProjectFile )
 			return false;
 		}
 
-		// Ask whether to copy the project or upgrade in-place
+		// Button labels for the upgrade dialog
 		TArray<FText> Buttons;
-		int32 OpenCopyButton = Buttons.Add(LOCTEXT("ProjectUpgradeCopy", "Create a copy"));
-		Buttons.Add(LOCTEXT("ProjectUpgradeInPlace", "Convert in-place"));
-		int32 SkipConversionButton = Buttons.Add(LOCTEXT("ProjectUpgradeSkipConversion", "Skip conversion"));
-		int32 CancelButton = Buttons.Add(LOCTEXT("ProjectUpgradeCancel", "Cancel"));
+		int32 OpenCopyButton = Buttons.Add(LOCTEXT("ProjectConvert_OpenCopy", "Open a copy"));
+		int32 OpenExistingButton = Buttons.Add(LOCTEXT("ProjectConvert_ConvertInPlace", "Convert in-place"));
+		int32 SkipConversionButton = Buttons.Add(LOCTEXT("ProjectConvert_SkipConversion", "Skip conversion"));
+		int32 CancelButton = Buttons.Add(LOCTEXT("ProjectConvert_Cancel", "Cancel"));
+
+		// Prompt for upgrading. Different message for code and content projects, since the process is a bit trickier for code.
+		int32 Selection;
+		if(ProjectStatus.bCodeBasedProject)
+		{
+			Selection = SVerbChoiceDialog::ShowModal(LOCTEXT("ProjectConversionTitle", "Convert Project"), LOCTEXT("ConvertCodeProjectPrompt", "This project was made with a different version of the Unreal Engine. Converting to this version will rebuild your code projects.\n\nNew features and improvements sometimes cause API changes, which may require you to modify your code before it compiles. Content saved with newer versions of the editor will not open in older versions.\n\nWe recommend you open a copy of your project to avoid damaging the original."), Buttons);
+		}
+		else
+		{
+			Selection = SVerbChoiceDialog::ShowModal(LOCTEXT("ProjectConversionTitle", "Convert Project"), LOCTEXT("ConvertContentProjectPrompt", "This project was made with a different version of the Unreal Engine.\n\nOpening it with this version of the editor may prevent it opening with the original editor, and may lose data. We recommend you open a copy to avoid damaging the original."), Buttons);
+		}
 
-		// Show the dialog
-		int32 Selection = SVerbChoiceDialog::ShowModal(LOCTEXT("ProjectConversionTitle", "Convert Project"), LOCTEXT("ConvertProjectPrompt", "This project was made with a different version of the editor. Opening with this version may be irreversible or lose data.\n\nIt is recommened that you create a copy before proceeding."), Buttons);
+		// Handle the selection
+		if(Selection == CancelButton)
+		{
+			return false;
+		}
 		if(Selection == OpenCopyButton)
 		{
 			FString NewProjectFile;
@@ -914,12 +928,6 @@ bool SProjectBrowser::OpenProject( const FString& InProjectFile )
 			}
 			ProjectFile = NewProjectFile;
 		}
-		else if(Selection == CancelButton)
-		{
-			return false;
-		}
-
-		// Unless it was skipped, do the conversion
 		if(Selection != SkipConversionButton)
 		{
 			// If it's a code-based project, generate project files and open visual studio after an upgrade
@@ -973,17 +981,13 @@ bool SProjectBrowser::OpenProject( const FString& InProjectFile )
 				}
 			}
 
-			// Update the game project to the latest version. This will prompt to check-out as necessary. We don't need to write the engine identifier yet, because it won't use the right .uprojectdirs logic.
-			if(!GameProjectUtils::UpdateGameProject(TEXT("")))
-			{
-				return false;
-			}
-
-			// Write the engine association
-			if(!FDesktopPlatformModule::Get()->SetEngineIdentifierForProject(ProjectFile, CurrentIdentifier))
+			// Update the game project to the latest version. This will prompt to check-out as necessary. We don't need to write the engine identifier directly, because it won't use the right .uprojectdirs logic.
+			if(!GameProjectUtils::UpdateGameProject(TEXT("")) || !FDesktopPlatformModule::Get()->SetEngineIdentifierForProject(ProjectFile, CurrentIdentifier))
 			{
-				FMessageDialog::Open(EAppMsgType::Ok, LOCTEXT("ProjectUpgradeFailure", "Project file could not be updated to latest version. Check that the file is writable."));
-				return false;
+				if(FMessageDialog::Open(EAppMsgType::YesNo, LOCTEXT("ProjectUpgradeFailure", "Project file could not be updated to latest version. Attempt to open anyway?")) != EAppReturnType::Yes)
+				{
+					return false;
+				}
 			}
 		}
 	}
diff --git a/Engine/Source/Editor/GameProjectGeneration/Private/SVerbChoiceDialog.cpp b/Engine/Source/Editor/GameProjectGeneration/Private/SVerbChoiceDialog.cpp
--- a/Engine/Source/Editor/GameProjectGeneration/Private/SVerbChoiceDialog.cpp
+++ b/Engine/Source/Editor/GameProjectGeneration/Private/SVerbChoiceDialog.cpp
@@ -16,7 +16,7 @@ int32 SVerbChoiceDialog::ShowModal( const FText& InTitle, const FText& InMessage
 		.ParentWindow(ModalWindow)
 		.Message( InMessage )
 		.Buttons( InButtons )
-		.WrapMessageAt(512.0f);
+		.WrapMessageAt(640.0f);
 
 	ModalWindow->SetContent( MessageBox );
 
@@ -72,7 +72,7 @@ void SVerbChoiceDialog::Construct( const FArguments& InArgs )
 							+ SHorizontalBox::Slot()
 								.FillWidth(1.0f)
 								.HAlign(HAlign_Left)
-								.VAlign(VAlign_Bottom)
+								.VAlign(VAlign_Center)
 								.Padding(12.0f)
 								[
 									SNew(SHyperlink)
@@ -84,7 +84,7 @@ void SVerbChoiceDialog::Construct( const FArguments& InArgs )
 							+ SHorizontalBox::Slot()
 								.AutoWidth()
 								.HAlign(HAlign_Right)
-								.VAlign(VAlign_Bottom)
+								.VAlign(VAlign_Center)
 								.Padding(2.f)
 								[
 									SAssignNew( ButtonBox, SUniformGridPanel )
