Format the version string in upgraded projects the same way as it's formatted by the launcher. Stops silly names like "Tappy Chicken 4.1 (4.2)", and prevents creating a separate DDC directory due to bug in DDC INI parsing.

diff --git a/Engine/Source/Editor/GameProjectGeneration/Private/GameProjectUtils.cpp b/Engine/Source/Editor/GameProjectGeneration/Private/GameProjectUtils.cpp
--- a/Engine/Source/Editor/GameProjectGeneration/Private/GameProjectUtils.cpp
+++ b/Engine/Source/Editor/GameProjectGeneration/Private/GameProjectUtils.cpp
@@ -1902,18 +1902,19 @@ bool GameProjectUtils::DuplicateProjectForUpgrade( const FString& InProjectFile,
 	for(int32 LastSpace; NewDirectoryName.FindLastChar(' ', LastSpace); )
 	{
 		const TCHAR *End = *NewDirectoryName + LastSpace + 1;
-		if(*End != '(')
+		if(End[0] != '4' || End[1] != '.' || !FChar::IsDigit(End[2]))
 		{
 			break;
 		}
 
-		End++;
-		while(FChar::IsDigit(*End) || *End == '.')
+		End += 3;
+
+		while(FChar::IsDigit(*End))
 		{
 			End++;
 		}
 
-		if(*End != ')' || *(End + 1) != 0)
+		if(*End != 0)
 		{
 			break;
 		}
@@ -1922,13 +1923,13 @@ bool GameProjectUtils::DuplicateProjectForUpgrade( const FString& InProjectFile,
 	}
 
 	// Append the new version number
-	NewDirectoryName += FString::Printf(TEXT(" (%s)"), *GEngineVersion.ToString(EVersionComponent::Minor));
+	NewDirectoryName += FString::Printf(TEXT(" %s"), *GEngineVersion.ToString(EVersionComponent::Minor));
 
 	// Find a directory name that doesn't exist
 	FString BaseDirectoryName = NewDirectoryName;
 	for(int32 Idx = 2; IFileManager::Get().DirectoryExists(*NewDirectoryName); Idx++)
 	{
-		NewDirectoryName = FString::Printf(TEXT("%s (%d)"), *BaseDirectoryName, Idx);
+		NewDirectoryName = FString::Printf(TEXT("%s - %d"), *BaseDirectoryName, Idx);
 	}
 
 	// Find all the root directory names
