MinimalAPI classes are now considered valid base classes in the class wizard

https://answers.unrealengine.com/questions/49569/i-cant-find-scenecapture2d-in-add-code-to-project.html

diff --git a/Engine/Source/Editor/GameProjectGeneration/Private/SNewClassDialog.cpp b/Engine/Source/Editor/GameProjectGeneration/Private/SNewClassDialog.cpp
--- a/Engine/Source/Editor/GameProjectGeneration/Private/SNewClassDialog.cpp
+++ b/Engine/Source/Editor/GameProjectGeneration/Private/SNewClassDialog.cpp
@@ -39,7 +39,8 @@ public:
 		const bool bNeedsAPI = !bIsExplicitlyUObject && !bIsInDestinationModule;
 
 		// You may not make a class that is not DLL exported.
-		const bool bHasAPI = InClass->HasAnyClassFlags(CLASS_RequiredAPI);
+		// MinimalAPI classes aren't compatible with the DLL export macro, but can still be used as a valid base
+		const bool bHasAPI = InClass->HasAnyClassFlags(CLASS_RequiredAPI) || InClass->HasAnyClassFlags(CLASS_MinimalAPI);
 
 		// @todo should we support interfaces?
 		const bool bIsInterface = InClass->IsChildOf(UInterface::StaticClass());
