Fixed a "todo" about a module name assumption

This now uses GameProjectUtils::GetCurrentModuleContextInfo(), which provides the same information as used by the rest of the class wizard validation.

diff --git a/Engine/Source/Editor/GameProjectGeneration/Private/SNewClassDialog.cpp b/Engine/Source/Editor/GameProjectGeneration/Private/SNewClassDialog.cpp
--- a/Engine/Source/Editor/GameProjectGeneration/Private/SNewClassDialog.cpp
+++ b/Engine/Source/Editor/GameProjectGeneration/Private/SNewClassDialog.cpp
@@ -30,10 +30,10 @@ public:
 		// UObject is special cased to be extensible since it would otherwise not be since it doesn't pass the API check (intrinsic class).
 		const bool bIsExplicitlyUObject = (InClass == UObject::StaticClass());
 
-		// @todo Assuming the game name is the same as the destination module name
-		const FString DestModuleName = FApp::GetGameName();
+		// Is this class in the same module as our current module?
+		const GameProjectUtils::FModuleContextInfo ModuleInfo = GameProjectUtils::GetCurrentModuleContextInfo();
 		const FString ClassModuleName = InClass->GetOutermost()->GetName().RightChop( FString(TEXT("/Script/")).Len() );
-		const bool bIsInDestinationModule = (DestModuleName == ClassModuleName);
+		const bool bIsInDestinationModule = (ModuleInfo.ModuleName == ClassModuleName);
 
 		// You need API if you are either not UObject itself and you are not in the destination module
 		const bool bNeedsAPI = !bIsExplicitlyUObject && !bIsInDestinationModule;
