Re-added support for custom categories in the Project Browser

TTP# 327759 - Cannot set Project Browser Category on user-created projects

The code for it had been removed during some other changes.

diff --git a/Engine/Source/Editor/GameProjectGeneration/Private/SProjectBrowser.cpp b/Engine/Source/Editor/GameProjectGeneration/Private/SProjectBrowser.cpp
--- a/Engine/Source/Editor/GameProjectGeneration/Private/SProjectBrowser.cpp
+++ b/Engine/Source/Editor/GameProjectGeneration/Private/SProjectBrowser.cpp
@@ -648,14 +648,17 @@ FText SProjectBrowser::GetSelectedProjectName() const
 
 FReply SProjectBrowser::FindProjects()
 {
+	enum class EProjectCategoryType : uint8
+	{
+		Sample,
+		UserDefined,
+	};
+
 	ProjectCategories.Empty();
 	CategoriesBox->ClearChildren();
 
-	const FText MyProjectsCategoryName = LOCTEXT("MyProjectsCategoryName", "My Projects");
-	const FText SamplesCategoryName = LOCTEXT("SamplesCategoryName", "Samples");
-
 	// Create a map of parent project folders to their category
-	TMap<FString, FText> ProjectFilesToCategory;
+	TMap<FString, EProjectCategoryType> ProjectFilesToCategoryType;
 
 	// Find all the engine installations
 	TMap<FString, FString> EngineInstallations;
@@ -670,7 +673,7 @@ FReply SProjectBrowser::FindProjects()
 		{
 			for(int Idx = 0; Idx < ProjectFiles.Num(); Idx++)
 			{
-				ProjectFilesToCategory.Add(ProjectFiles[Idx], MyProjectsCategoryName);
+				ProjectFilesToCategoryType.Add(ProjectFiles[Idx], EProjectCategoryType::UserDefined);
 			}
 		}
 	}
@@ -680,7 +683,7 @@ FReply SProjectBrowser::FindProjects()
 	FDesktopPlatformModule::Get()->EnumerateLauncherSampleProjects(LauncherSampleProjects);
 	for(int32 Idx = 0; Idx < LauncherSampleProjects.Num(); Idx++)
 	{
-		ProjectFilesToCategory.Add(LauncherSampleProjects[Idx], SamplesCategoryName);
+		ProjectFilesToCategoryType.Add(LauncherSampleProjects[Idx], EProjectCategoryType::Sample);
 	}
 
 	// Add all the native project files we can find, and automatically filter them depending on their directory
@@ -689,25 +692,28 @@ FReply SProjectBrowser::FindProjects()
 	{
 		if(!NativeProjectFiles[Idx].Contains(TEXT("/Templates/")))
 		{
-			const FText &CategoryName = NativeProjectFiles[Idx].Contains(TEXT("/Samples/"))? SamplesCategoryName : MyProjectsCategoryName;
-			ProjectFilesToCategory.Add(NativeProjectFiles[Idx], CategoryName);
+			const EProjectCategoryType ProjectCategoryType = NativeProjectFiles[Idx].Contains(TEXT("/Samples/")) ? EProjectCategoryType::Sample : EProjectCategoryType::UserDefined;
+			ProjectFilesToCategoryType.Add(NativeProjectFiles[Idx], ProjectCategoryType);
 		}
 	}
 
 	// Normalize all the filenames and make sure there are no duplicates
-	TMap<FString, FText> AbsoluteProjectFilesToCategory;
-	for(TMap<FString, FText>::TConstIterator Iter(ProjectFilesToCategory); Iter; ++Iter)
+	TMap<FString, EProjectCategoryType> AbsoluteProjectFilesToCategory;
+	for(TMap<FString, EProjectCategoryType>::TConstIterator Iter(ProjectFilesToCategoryType); Iter; ++Iter)
 	{
 		FString AbsoluteFile = IFileManager::Get().ConvertToAbsolutePathForExternalAppForRead(*Iter.Key());
 		AbsoluteProjectFilesToCategory.Add(AbsoluteFile, Iter.Value());
 	}
 
+	const FText MyProjectsCategoryName = LOCTEXT("MyProjectsCategoryName", "My Projects");
+	const FText SamplesCategoryName = LOCTEXT("SamplesCategoryName", "Samples");
+
 	// Add all the discovered projects to the list
 	const FString EngineIdentifier = FDesktopPlatformModule::Get()->GetCurrentEngineIdentifier();
-	for(TMap<FString, FText>::TConstIterator Iter(AbsoluteProjectFilesToCategory); Iter; ++Iter)
+	for(TMap<FString, EProjectCategoryType>::TConstIterator Iter(AbsoluteProjectFilesToCategory); Iter; ++Iter)
 	{
 		const FString& ProjectFilename = *Iter.Key();
-		const FText& DetectedCategory = Iter.Value();
+		const EProjectCategoryType DetectedCategoryType = Iter.Value();
 		if ( FPaths::FileExists(ProjectFilename) )
 		{
 			FProjectStatus ProjectStatus;
@@ -739,11 +745,21 @@ FReply SProjectBrowser::FindProjects()
 				FText ProjectCategory;
 				if ( ProjectStatus.bSignedSampleProject )
 				{
+					// Signed samples can't override their category name
 					ProjectCategory = SamplesCategoryName;
 				}
 				else
 				{
-					ProjectCategory = DetectedCategory;
+					if(ProjectStatus.Category.IsEmpty())
+					{
+						// Not category specified, so use the category for the detected project type
+						ProjectCategory = (DetectedCategoryType == EProjectCategoryType::Sample) ? SamplesCategoryName : MyProjectsCategoryName;
+					}
+					else
+					{
+						// Use the user defined category
+						ProjectCategory = FText::FromString(ProjectStatus.Category);
+					}
 				}
 
 				FString ProjectEngineIdentifier;
