Fix BasicCode project not referencing custom game mode
TTP #336609

diff --git a/Engine/Source/Editor/GameProjectGeneration/Private/GameProjectUtils.cpp b/Engine/Source/Editor/GameProjectGeneration/Private/GameProjectUtils.cpp
--- a/Engine/Source/Editor/GameProjectGeneration/Private/GameProjectUtils.cpp
+++ b/Engine/Source/Editor/GameProjectGeneration/Private/GameProjectUtils.cpp
@@ -1462,6 +1462,11 @@ bool GameProjectUtils::GenerateConfigFiles(const FString& NewProjectPath, const
 				FileContents += TEXT("[/Script/EngineSettings.GameMapsSettings]") LINE_TERMINATOR;
 				FileContents += FString::Printf(TEXT("EditorStartupMap=%s") LINE_TERMINATOR, *MapPackagePath);
 				FileContents += FString::Printf(TEXT("GameDefaultMap=%s") LINE_TERMINATOR, *MapPackagePath);
+				if (bShouldGenerateCode)
+				{
+					FileContents += FString::Printf(TEXT("GlobalDefaultGameMode=\"/Script/%s.%sGameMode\"") LINE_TERMINATOR, *NewProjectName, *NewProjectName);
+					FileContents += FString::Printf(TEXT("GlobalDefaultServerGameMode=\"/Script/%s.%sGameMode\"") LINE_TERMINATOR, *NewProjectName, *NewProjectName);
+				}
 			}
 		}
 
@@ -1483,14 +1488,6 @@ bool GameProjectUtils::GenerateConfigFiles(const FString& NewProjectPath, const
 		FileContents += FString::Printf( TEXT("ProjectID=%s") LINE_TERMINATOR, *FGuid::NewGuid().ToString() );
 		FileContents += LINE_TERMINATOR;
 
-		if ( bShouldGenerateCode )
-		{
-			FileContents += TEXT("[/Script/Engine.WorldSettings]") LINE_TERMINATOR;
-			FileContents += FString::Printf(TEXT("GlobalDefaultGameMode=\"/Script/%s.%sGameMode\"") LINE_TERMINATOR, *NewProjectName, *NewProjectName);
-			FileContents += FString::Printf(TEXT("GlobalDefaultServerGameMode=\"/Script/%s.%sGameMode\"") LINE_TERMINATOR, *NewProjectName, *NewProjectName);
-			FileContents += LINE_TERMINATOR;
-		}
-
 		if ( WriteOutputFile(DefaultGameIniFilename, FileContents, OutFailReason) )
 		{
 			OutCreatedFiles.Add(DefaultGameIniFilename);
