UE4: Renamed FStatusMessageContext to FScopedSlowTask which describes its purpose much more clearly.

diff --git a/Engine/Source/Developer/AssetTools/Private/AssetFixUpRedirectors.cpp b/Engine/Source/Developer/AssetTools/Private/AssetFixUpRedirectors.cpp
--- a/Engine/Source/Developer/AssetTools/Private/AssetFixUpRedirectors.cpp
+++ b/Engine/Source/Developer/AssetTools/Private/AssetFixUpRedirectors.cpp
@@ -138,7 +138,7 @@ bool FAssetFixUpRedirectors::UpdatePackageStatus(const TArray<FRedirectorRefs>&
 void FAssetFixUpRedirectors::LoadReferencingPackages(TArray<FRedirectorRefs>& RedirectorsToFix, TArray<UPackage*>& OutReferencingPackagesToSave) const
 {
 	const bool bAllowNewSlowTask = true;
-	FStatusMessageContext SlowTaskMessage( LOCTEXT( "LoadingReferencingPackages", "Loading Referencing Packages..." ), bAllowNewSlowTask );
+	FScopedSlowTask SlowTaskMessage( LOCTEXT( "LoadingReferencingPackages", "Loading Referencing Packages..." ), bAllowNewSlowTask );
 	ISourceControlProvider& SourceControlProvider = ISourceControlModule::Get().GetProvider();
 
 	// Load all packages that reference each redirector, if possible
diff --git a/Engine/Source/Developer/AssetTools/Private/AssetTools.cpp b/Engine/Source/Developer/AssetTools/Private/AssetTools.cpp
--- a/Engine/Source/Developer/AssetTools/Private/AssetTools.cpp
+++ b/Engine/Source/Developer/AssetTools/Private/AssetTools.cpp
@@ -1159,7 +1159,7 @@ void FAssetTools::PerformMigratePackages(TArray<FName> PackageNamesToMigrate) co
 	TSet<FName> AllPackageNamesToMove;
 	{
 		const bool bAllowNewSlowTask = true;
-		FStatusMessageContext SlowTaskGatheringDependenciesMessage( LOCTEXT( "MigratePackages_GatheringDependencies", "Gathering Dependencies..." ) );
+		FScopedSlowTask SlowTaskGatheringDependenciesMessage( LOCTEXT( "MigratePackages_GatheringDependencies", "Gathering Dependencies..." ) );
 
 		for ( auto PackageIt = PackageNamesToMigrate.CreateConstIterator(); PackageIt; ++PackageIt )
 		{
@@ -1256,7 +1256,7 @@ void FAssetTools::MigratePackages_ReportConfirmed(TArray<FString> ConfirmedPacka
 
 	// Copy all specified assets and their dependencies to the destination folder
 	const bool bAllowNewSlowTask = true;
-	FStatusMessageContext SlowTaskMessage( LOCTEXT( "MigratePackages_CopyingFiles", "Copying Files..." ), bAllowNewSlowTask );
+	FScopedSlowTask SlowTaskMessage( LOCTEXT( "MigratePackages_CopyingFiles", "Copying Files..." ), bAllowNewSlowTask );
 
 	EAppReturnType::Type LastResponse = EAppReturnType::Yes;
 	TArray<FString> SuccessfullyCopiedFiles;
diff --git a/Engine/Source/Editor/GameProjectGeneration/Private/GameProjectGenerationModule.cpp b/Engine/Source/Editor/GameProjectGeneration/Private/GameProjectGenerationModule.cpp
--- a/Engine/Source/Editor/GameProjectGeneration/Private/GameProjectGenerationModule.cpp
+++ b/Engine/Source/Editor/GameProjectGeneration/Private/GameProjectGenerationModule.cpp
@@ -59,7 +59,7 @@ bool FGameProjectGenerationModule::UpdateGameProject(const FString& EngineIdenti
 bool FGameProjectGenerationModule::UpdateCodeProject(FText& OutFailReason)
 {
 	const bool bAllowNewSlowTask = true;
-	FStatusMessageContext SlowTaskMessage( LOCTEXT( "UpdatingCodeProject", "Updating code project..." ), bAllowNewSlowTask );
+	FScopedSlowTask SlowTaskMessage( LOCTEXT( "UpdatingCodeProject", "Updating code project..." ), bAllowNewSlowTask );
 
 	return GameProjectUtils::GenerateCodeProjectFiles(FPaths::GetProjectFilePath(), OutFailReason);
 }
diff --git a/Engine/Source/Editor/GameProjectGeneration/Private/GameProjectUtils.cpp b/Engine/Source/Editor/GameProjectGeneration/Private/GameProjectUtils.cpp
--- a/Engine/Source/Editor/GameProjectGeneration/Private/GameProjectUtils.cpp
+++ b/Engine/Source/Editor/GameProjectGeneration/Private/GameProjectUtils.cpp
@@ -539,7 +539,7 @@ bool GameProjectUtils::CopyStarterContent(const FString& DestProjectFolder, FTex
 		const bool bAllowNewSlowTask = false;
 		FFormatNamedArguments Args;
 		Args.Add(TEXT("SrcFilename"), FText::FromString(FPaths::GetCleanFilename(SrcFilename)));
-		FStatusMessageContext SlowTaskMessage(FText::Format(LOCTEXT("CreatingProjectStatus_CopyingFile", "Copying File {SrcFilename}..."), Args), bAllowNewSlowTask);
+		FScopedSlowTask SlowTaskMessage(FText::Format(LOCTEXT("CreatingProjectStatus_CopyingFile", "Copying File {SrcFilename}..."), Args), bAllowNewSlowTask);
 
 		FString FileRelPath = FPaths::GetPath(SrcFilename);
 		FPaths::MakePathRelativeTo(FileRelPath, *SrcFolder);
@@ -576,7 +576,7 @@ bool GameProjectUtils::CreateProject(const FString& NewProjectFile, const FStrin
 	}
 
 	const bool bAllowNewSlowTask = true;
-	FStatusMessageContext SlowTaskMessage( LOCTEXT( "CreatingProjectStatus", "Creating project..." ), bAllowNewSlowTask );
+	FScopedSlowTask SlowTaskMessage( LOCTEXT( "CreatingProjectStatus", "Creating project..." ), bAllowNewSlowTask );
 
 	bool bProjectCreationSuccessful = false;
 	FString TemplateName;
@@ -607,7 +607,7 @@ bool GameProjectUtils::CreateProject(const FString& NewProjectFile, const FStrin
 bool GameProjectUtils::BuildGameBinaries(const FString& ProjectFilename, FText& OutFailReason)
 {
 	const bool bAllowNewSlowTask = true;
-	FStatusMessageContext SlowTaskMessage( LOCTEXT( "BuildingProjectStatus", "Building project..." ), bAllowNewSlowTask );
+	FScopedSlowTask SlowTaskMessage( LOCTEXT( "BuildingProjectStatus", "Building project..." ), bAllowNewSlowTask );
 
 	// Compile the *editor* for the project
 	if ( FModuleManager::Get().CompileGameProjectEditor(ProjectFilename, *GLog) )
@@ -1037,7 +1037,7 @@ bool GameProjectUtils::CreateProjectFromTemplate(const FString& NewProjectFile,
 		const bool bAllowNewSlowTask = false;
 		FFormatNamedArguments Args;
 		Args.Add( TEXT("SrcFilename"), FText::FromString( FPaths::GetCleanFilename(SrcFilename) ) );
-		FStatusMessageContext SlowTaskMessage( FText::Format( LOCTEXT( "CreatingProjectStatus_CopyingFile", "Copying File {SrcFilename}..." ), Args ), bAllowNewSlowTask );
+		FScopedSlowTask SlowTaskMessage( FText::Format( LOCTEXT( "CreatingProjectStatus_CopyingFile", "Copying File {SrcFilename}..." ), Args ), bAllowNewSlowTask );
 
 		// Retarget any folders that were chosen to be renamed by choosing a new destination subpath now
 		FString DestFileSubpathWithoutFilename = FPaths::GetPath(SrcFileSubpath) + TEXT("/");
@@ -2568,7 +2568,7 @@ bool GameProjectUtils::AddCodeToProject_Internal(const FString& NewClassName, co
 	}
 
 	const bool bAllowNewSlowTask = true;
-	FStatusMessageContext SlowTaskMessage( LOCTEXT( "AddingCodeToProject", "Adding code to project..." ), bAllowNewSlowTask );
+	FScopedSlowTask SlowTaskMessage( LOCTEXT( "AddingCodeToProject", "Adding code to project..." ), bAllowNewSlowTask );
 
 	// If the project does not already contain code, add the primary game module
 	TArray<FString> CreatedFiles;
diff --git a/Engine/Source/Editor/UnrealEd/Private/ObjectTools.cpp b/Engine/Source/Editor/UnrealEd/Private/ObjectTools.cpp
--- a/Engine/Source/Editor/UnrealEd/Private/ObjectTools.cpp
+++ b/Engine/Source/Editor/UnrealEd/Private/ObjectTools.cpp
@@ -3466,7 +3466,7 @@ namespace ThumbnailTools
 			if ( UMaterial* InMaterial = Cast<UMaterial>(InObject) )
 			{
 				const bool bAllowNewSlowTask = true;
-				FStatusMessageContext SlowTaskMessage( NSLOCTEXT( "ObjectTools", "FinishingCompilationStatus", "Finishing Shader Compilation..." ), bAllowNewSlowTask );
+				FScopedSlowTask SlowTaskMessage( NSLOCTEXT( "ObjectTools", "FinishingCompilationStatus", "Finishing Shader Compilation..." ), bAllowNewSlowTask );
 
 				// Block until the shader maps that we will save have finished being compiled
 				InMaterial->GetMaterialResource(GRHIFeatureLevel)->FinishCompilation();
diff --git a/Engine/Source/Runtime/Core/Public/Misc/FeedbackContext.h b/Engine/Source/Runtime/Core/Public/Misc/FeedbackContext.h
--- a/Engine/Source/Runtime/Core/Public/Misc/FeedbackContext.h
+++ b/Engine/Source/Runtime/Core/Public/Misc/FeedbackContext.h
@@ -75,14 +75,14 @@ public:
 /**
  * Helper class to display a status update message in the editor.
  */
-class FStatusMessageContext
+class FScopedSlowTask
 {
 public:
 
 	/**
 	 * Updates the status message displayed to the user.
 	 */
-	explicit FStatusMessageContext( const FText& InMessage, bool bAllowNewSlowTask = false )
+	explicit FScopedSlowTask( const FText& InMessage, bool bAllowNewSlowTask = false )
 	{
 		bIsBeginSlowTask = false;
 
@@ -106,7 +106,7 @@ public:
 	/**
 	 * Ensures the status context is popped off the stack.
 	 */
-	~FStatusMessageContext()
+	~FScopedSlowTask()
 	{
 		if ( GIsEditor && !IsRunningCommandlet() && IsInGameThread() )
 		{
diff --git a/Engine/Source/Runtime/Engine/Private/StaticMesh.cpp b/Engine/Source/Runtime/Engine/Private/StaticMesh.cpp
--- a/Engine/Source/Runtime/Engine/Private/StaticMesh.cpp
+++ b/Engine/Source/Runtime/Engine/Private/StaticMesh.cpp
@@ -1121,11 +1121,11 @@ static const FString& GetStaticMeshDerivedDataVersion()
 	return CachedVersionString;
 }
 
-class FStaticMeshStatusMessageContext : public FStatusMessageContext
+class FStaticMeshStatusMessageContext : public FScopedSlowTask
 {
 public:
 	explicit FStaticMeshStatusMessageContext(const FText& InMessage)
-		: FStatusMessageContext(InMessage)
+		: FScopedSlowTask(InMessage)
 	{
 		UE_LOG(LogStaticMesh,Log,TEXT("%s"),*InMessage.ToString());
 	}
diff --git a/Engine/Source/Runtime/Engine/Private/TextureDerivedData.cpp b/Engine/Source/Runtime/Engine/Private/TextureDerivedData.cpp
--- a/Engine/Source/Runtime/Engine/Private/TextureDerivedData.cpp
+++ b/Engine/Source/Runtime/Engine/Private/TextureDerivedData.cpp
@@ -445,11 +445,11 @@ static void PutDerivedDataInCache(
 
 #if WITH_EDITORONLY_DATA
 
-class FTextureStatusMessageContext : public FStatusMessageContext
+class FTextureStatusMessageContext : public FScopedSlowTask
 {
 public:
 	explicit FTextureStatusMessageContext(const FText& InMessage)
-		: FStatusMessageContext(InMessage)
+		: FScopedSlowTask(InMessage)
 	{
 		UE_LOG(LogTexture,Display,TEXT("%s"),*InMessage.ToString());
 	}
