Remove an IsRocket() check - open the non-foreign solution file for non-foreign projects.

diff --git a/Engine/Source/Editor/GameProjectGeneration/Private/GameProjectUtils.cpp b/Engine/Source/Editor/GameProjectGeneration/Private/GameProjectUtils.cpp
--- a/Engine/Source/Editor/GameProjectGeneration/Private/GameProjectUtils.cpp
+++ b/Engine/Source/Editor/GameProjectGeneration/Private/GameProjectUtils.cpp
@@ -15,6 +15,8 @@
 
 #include "ClassIconFinder.h"
 
+#include "UProjectInfo.h"
+
 #define LOCTEXT_NAMESPACE "GameProjectUtils"
 
 #define MAX_PROJECT_PATH_BUFFER_SPACE 130 // Leave a reasonable buffer of additional characters to account for files created in the content directory during or after project generation
@@ -453,8 +455,9 @@ bool GameProjectUtils::OpenCodeIDE(const FString& ProjectFile, FText& OutFailRea
 		return false;
 	}
 
+	// Check whether this project is a foreign project. Don't use the cached project dictionary; we may have just created a new project.
 	bool bIsInRootFolder = false;
-	if ( !FRocketSupport::IsRocket() )
+	if ( !FUProjectDictionary(FPaths::RootDir()).IsForeignProject(ProjectFile) )
 	{
 		// If we are in the UE4 root, just open the UE4.sln file, otherwise open the generated one.
 		FString AbsoluteProjectParentFolder = IFileManager::Get().ConvertToAbsolutePathForExternalAppForRead(*FPaths::GetPath(FPaths::GetPath(ProjectFile)));
