Generating a new class will now include the module header relative to a known include path

If the module is inside the Public or Classes folder, it is included directly (since all these folders are on the include path), otherwise it's included relative to the module source root (which is now on the include paths due to a recent UBT change).

diff --git a/Engine/Source/Editor/GameProjectGeneration/Private/GameProjectUtils.cpp b/Engine/Source/Editor/GameProjectGeneration/Private/GameProjectUtils.cpp
--- a/Engine/Source/Editor/GameProjectGeneration/Private/GameProjectUtils.cpp
+++ b/Engine/Source/Editor/GameProjectGeneration/Private/GameProjectUtils.cpp
@@ -753,7 +753,7 @@ void GameProjectUtils::OpenAddCodeToProjectDialog()
 	}
 }
 
-bool GameProjectUtils::IsValidClassNameForCreation(const FString& NewClassName, FText& OutFailReason)
+bool GameProjectUtils::IsValidClassNameForCreation(const FString& NewClassName, const FModuleContextInfo& ModuleInfo, FText& OutFailReason)
 {
 	if ( NewClassName.IsEmpty() )
 	{
@@ -801,12 +801,9 @@ bool GameProjectUtils::IsValidClassNameForCreation(const FString& NewClassName,
 	}
 
 	// Look for a duplicate class on disk in their project
-	TArray<FString> Filenames;
-	IFileManager::Get().FindFilesRecursive(Filenames, *FPaths::GameSourceDir(), TEXT("*.h"), true, false, false);
-	for ( auto FileIt = Filenames.CreateConstIterator(); FileIt; ++FileIt )
 	{
-		const FString& File = *FileIt;
-		if ( NewClassName == FPaths::GetBaseFilename(File) )
+		FString UnusedFoundPath;
+		if ( FindSourceFileInProject(NewClassName + ".h", UnusedFoundPath, ModuleInfo) )
 		{
 			FFormatNamedArguments Args;
 			Args.Add( TEXT("NewClassName"), FText::FromString( NewClassName ) );
@@ -2216,10 +2213,62 @@ bool GameProjectUtils::GenerateClassCPPFile(const FString& NewCPPFileName, const
 		PropertyOverridesStr += *PropertyOverrides[OverrideIdx];
 	}
 
+	// Calculate the correct include path for the module header
+	FString ModuleIncludePath;
+	if(FindSourceFileInProject(ModuleName + ".h", ModuleIncludePath, ModuleInfo))
+	{
+		// Work out where the module header is; 
+		// if it's Public then we can include it without any path since all Public and Classes folders are on the include path
+		// if it's located elsewhere, then we'll need to include it relative to the module source root as we can't guarantee 
+		// that other folders are on the include paths
+		FString UnusedModuleName;
+		EClassLocation ModuleLocation;
+		if(GetClassLocation(ModuleIncludePath, UnusedModuleName, ModuleLocation, ModuleInfo))
+		{
+			if(ModuleLocation == EClassLocation::Public || ModuleLocation == EClassLocation::Classes)
+			{
+				ModuleIncludePath = ModuleName + ".h";
+			}
+			else
+			{
+				// If the path to our new class is the same as the path to the module, we can include it directly
+				const FString ModulePath = FPaths::ConvertRelativePathToFull(FPaths::GetPath(ModuleIncludePath));
+				const FString ClassPath = FPaths::ConvertRelativePathToFull(FPaths::GetPath(NewCPPFileName));
+				if(ModulePath == ClassPath)
+				{
+					ModuleIncludePath = ModuleName + ".h";
+				}
+				else
+				{
+					// Make the include relative to the source root of the module
+					const FString ModuleSourceRoot = GetSourceRootPath(true/*bIncludeModuleName*/, ModuleInfo);
+					
+					// Updates ModuleIncludePath internally
+					if(!FPaths::MakePathRelativeTo(ModuleIncludePath, *ModuleSourceRoot))
+					{
+						// Failed; just assume we can include it without any relative path
+						ModuleIncludePath = ModuleName + ".h";
+					}
+				}
+			}
+		}
+		else
+		{
+			// Failed; just assume we can include it without any relative path
+			ModuleIncludePath = ModuleName + ".h";
+		}
+	}
+	else
+	{
+		// This could potentially fail when generating new projects if the module file hasn't yet been created; just assume we can include it without any relative path
+		ModuleIncludePath = ModuleName + ".h";
+	}
+
 	// Not all of these will exist in every class template
 	FString FinalOutput = Template.Replace(TEXT("%COPYRIGHT_LINE%"), *MakeCopyrightLine(), ESearchCase::CaseSensitive);
 	FinalOutput = FinalOutput.Replace(TEXT("%UNPREFIXED_CLASS_NAME%"), *UnPrefixedClassName, ESearchCase::CaseSensitive);
 	FinalOutput = FinalOutput.Replace(TEXT("%MODULE_NAME%"), *ModuleName, ESearchCase::CaseSensitive);
+	FinalOutput = FinalOutput.Replace(TEXT("%MODULE_INCLUDE_PATH%"), *ModuleIncludePath, ESearchCase::CaseSensitive);
 	FinalOutput = FinalOutput.Replace(TEXT("%PREFIXED_CLASS_NAME%"), *PrefixedClassName, ESearchCase::CaseSensitive);
 	FinalOutput = FinalOutput.Replace(TEXT("%PROPERTY_OVERRIDES%"), *PropertyOverridesStr, ESearchCase::CaseSensitive);
 	FinalOutput = FinalOutput.Replace(TEXT("%ADDITIONAL_MEMBER_DEFINITIONS%"), *AdditionalMemberDefinitions, ESearchCase::CaseSensitive);
@@ -2547,7 +2596,10 @@ bool GameProjectUtils::AddCodeToProject_Internal(const FString& NewClassName, co
 	const FString CleanClassName = ParentClassInfo.GetCleanClassName(NewClassName);
 	const FString FinalClassName = ParentClassInfo.GetFinalClassName(NewClassName);
 
-	if ( !IsValidClassNameForCreation(FinalClassName, OutFailReason) )
+	// Get the context info for the current project
+	const FModuleContextInfo ModuleInfo = GetCurrentModuleContextInfo();
+
+	if ( !IsValidClassNameForCreation(FinalClassName, ModuleInfo, OutFailReason) )
 	{
 		return false;
 	}
@@ -2558,9 +2610,6 @@ bool GameProjectUtils::AddCodeToProject_Internal(const FString& NewClassName, co
 		return false;
 	}
 
-	// Get the context info for the current project
-	const FModuleContextInfo ModuleInfo = GetCurrentModuleContextInfo();
-
 	FString ModuleName;
 	FString NewHeaderPath;
 	FString NewCppPath;
@@ -2649,31 +2698,20 @@ bool GameProjectUtils::AddCodeToProject_Internal(const FString& NewClassName, co
 	return true;
 }
 
-int32 GameProjectUtils::CalculateSubfolderCount(const FString& Filepath, const FString& RelativeRoot)
+bool GameProjectUtils::FindSourceFileInProject(const FString& InFilename, FString& OutPath, const FModuleContextInfo& ModuleInfo)
 {
-	int32 ModuleNameIndex = Filepath.Find(RelativeRoot, ESearchCase::IgnoreCase, ESearchDir::FromEnd);
-	if (ModuleNameIndex == INDEX_NONE)
-	{
-		return 0;
-	}
-
-	FString WorkingString = Filepath.Mid(ModuleNameIndex + RelativeRoot.Len() + 1);
-
-	int Depth = 0;
-	int32 SlashIndex;
-	while (WorkingString.FindChar(TEXT('/'), SlashIndex))
+	TArray<FString> Filenames;
+	const FString FilenameWidcard = TEXT("*") + InFilename;
+	IFileManager::Get().FindFilesRecursive(Filenames, *ModuleInfo.ModuleSourcePath, *FilenameWidcard, true, false, false);
+	
+	if(Filenames.Num())
 	{
-		// because input is not sanitized against multiple sequential slashes (e.g. '////')
-		// only consider slashes that don't appear at beginning of the working string
-		if (SlashIndex > 0)
-		{
-			Depth++;
-		}
-
-		WorkingString = WorkingString.Mid(SlashIndex + 1);
+		// Assume it's the first match (we should really only find a single file with a given name within a project anyway)
+		OutPath = Filenames[0];
+		return true;
 	}
 
-	return Depth;
+	return false;
 }
 
 #undef LOCTEXT_NAMESPACE
diff --git a/Engine/Source/Editor/GameProjectGeneration/Private/GameProjectUtils.h b/Engine/Source/Editor/GameProjectGeneration/Private/GameProjectUtils.h
--- a/Engine/Source/Editor/GameProjectGeneration/Private/GameProjectUtils.h
+++ b/Engine/Source/Editor/GameProjectGeneration/Private/GameProjectUtils.h
@@ -148,7 +148,7 @@ public:
 	static void OpenAddCodeToProjectDialog();
 
 	/** Returns true if the specified class name is properly formed and does not conflict with another class */
-	static bool IsValidClassNameForCreation(const FString& NewClassName, FText& OutFailReason);
+	static bool IsValidClassNameForCreation(const FString& NewClassName, const FModuleContextInfo& ModuleInfo, FText& OutFailReason);
 
 	/** Adds new source code to the project. When returning true, OutSyncFileAndLineNumber will be the the preferred target file to sync in the users code editing IDE, formatted for use with GenericApplication::GotoLineInSource */
 	static bool AddCodeToProject(const FString& NewClassName, const FString& NewClassPath, const FNewClassInfo ParentClassInfo, FString& OutHeaderFilePath, FString& OutCppFilePath, FText& OutFailReason);
@@ -295,7 +295,6 @@ private:
 	/** Creates the batch file for launching the editor or game */
 	static bool GenerateLaunchBatchFile(const FString& ProjectName, const FString& ProjectFolder, bool bLaunchEditor, TArray<FString>& OutCreatedFiles, FText& OutFailReason);
 
-
 	/** Returns the contents of the specified template file */
 	static bool ReadTemplateFile(const FString& TemplateFileName, FString& OutFileContents, FText& OutFailReason);
 
@@ -375,7 +374,8 @@ private:
 	/** Handler for the user confirming they've read the name legnth warning */
 	static void OnWarningReasonOk();
 
-	static int32 CalculateSubfolderCount(const FString& Filepath, const FString& RelativeRoot);
+	/** Given a source file name, find its location within the project */
+	static bool FindSourceFileInProject(const FString& InFilename, FString& OutPath, const FModuleContextInfo& ModuleInfo);
 
 private:
 	static TWeakPtr<SNotificationItem> UpdateGameProjectNotification;
diff --git a/Engine/Source/Editor/GameProjectGeneration/Private/SNewClassDialog.cpp b/Engine/Source/Editor/GameProjectGeneration/Private/SNewClassDialog.cpp
--- a/Engine/Source/Editor/GameProjectGeneration/Private/SNewClassDialog.cpp
+++ b/Engine/Source/Editor/GameProjectGeneration/Private/SNewClassDialog.cpp
@@ -971,7 +971,7 @@ void SNewClassDialog::UpdateInputValidity()
 	// Validate the class name only if the path is valid
 	if ( bLastInputValidityCheckSuccessful )
 	{
-		bLastInputValidityCheckSuccessful = GameProjectUtils::IsValidClassNameForCreation(NewClassName, LastInputValidityErrorText);
+		bLastInputValidityCheckSuccessful = GameProjectUtils::IsValidClassNameForCreation(NewClassName, ModuleInfo, LastInputValidityErrorText);
 	}
 
 	LastPeriodicValidityCheckTime = FSlateApplication::Get().GetCurrentTime();
