Fixed opening solutions when creating a non-foreign project in a subdirectory of the engine root.

diff --git a/Engine/Source/Editor/GameProjectGeneration/Private/GameProjectUtils.cpp b/Engine/Source/Editor/GameProjectGeneration/Private/GameProjectUtils.cpp
--- a/Engine/Source/Editor/GameProjectGeneration/Private/GameProjectUtils.cpp
+++ b/Engine/Source/Editor/GameProjectGeneration/Private/GameProjectUtils.cpp
@@ -456,39 +456,20 @@ bool GameProjectUtils::OpenCodeIDE(const FString& ProjectFile, FText& OutFailRea
 	}
 
 	// Check whether this project is a foreign project. Don't use the cached project dictionary; we may have just created a new project.
-	bool bIsInRootFolder = false;
-	if ( !FUProjectDictionary(FPaths::RootDir()).IsForeignProject(ProjectFile) )
-	{
-		// If we are in the UE4 root, just open the UE4.sln file, otherwise open the generated one.
-		FString AbsoluteProjectParentFolder = IFileManager::Get().ConvertToAbsolutePathForExternalAppForRead(*FPaths::GetPath(FPaths::GetPath(ProjectFile)));
-		FString AbsoluteRootPath = IFileManager::Get().ConvertToAbsolutePathForExternalAppForRead(*FPaths::RootDir());
-
-		if ( AbsoluteProjectParentFolder.Right(1) != TEXT("/") )
-		{
-			AbsoluteProjectParentFolder += TEXT("/");
-		}
-
-		if ( AbsoluteRootPath.Right(1) != TEXT("/") )
-		{
-			AbsoluteRootPath += TEXT("/");
-		}
-		
-		bIsInRootFolder = (AbsoluteProjectParentFolder == AbsoluteRootPath);
-	}
-
 	FString SolutionFolder;
 	FString SolutionFilenameWithoutExtension;
-	if ( bIsInRootFolder )
+	if( FUProjectDictionary(FPaths::RootDir()).IsForeignProject(ProjectFile) )
 	{
-		SolutionFolder = IFileManager::Get().ConvertToAbsolutePathForExternalAppForRead(*FPaths::RootDir());
-		SolutionFilenameWithoutExtension = TEXT("UE4");
+		SolutionFolder = IFileManager::Get().ConvertToAbsolutePathForExternalAppForRead(*FPaths::GetPath(ProjectFile));
+		SolutionFilenameWithoutExtension = FPaths::GetBaseFilename(ProjectFile);
 	}
 	else
 	{
-		SolutionFolder = IFileManager::Get().ConvertToAbsolutePathForExternalAppForRead(*FPaths::GetPath(ProjectFile));
-		SolutionFilenameWithoutExtension = FPaths::GetBaseFilename(ProjectFile);
+		SolutionFolder = IFileManager::Get().ConvertToAbsolutePathForExternalAppForRead(*FPaths::RootDir());
+		SolutionFilenameWithoutExtension = TEXT("UE4");
 	}
 
+	// Get the solution filename
 	FString CodeSolutionFile;
 #if PLATFORM_WINDOWS
 	CodeSolutionFile = SolutionFilenameWithoutExtension + TEXT(".sln");
@@ -499,8 +480,8 @@ bool GameProjectUtils::OpenCodeIDE(const FString& ProjectFile, FText& OutFailRea
 	return false;
 #endif
 
+	// Open the solution with the default application
 	const FString FullPath = FPaths::Combine(*SolutionFolder, *CodeSolutionFile);
-
 #if PLATFORM_MAC || PLATFORM_LINUX
 	if ( IFileManager::Get().DirectoryExists(*FullPath) )
 #else
