[GitHub] 263 : Search for project templates in plugin directories

diff --git a/Engine/Source/Editor/GameProjectGeneration/Private/SNewProjectWizard.cpp b/Engine/Source/Editor/GameProjectGeneration/Private/SNewProjectWizard.cpp
--- a/Engine/Source/Editor/GameProjectGeneration/Private/SNewProjectWizard.cpp
+++ b/Engine/Source/Editor/GameProjectGeneration/Private/SNewProjectWizard.cpp
@@ -908,6 +908,21 @@ void SNewProjectWizard::FindTemplateProjects()
 	// @todo rocket make template folder locations extensible.
 	TemplateRootFolders.Add( FPaths::RootDir() + TEXT("Templates") );
 
+	// allow plugins to define templates
+	TArray<FPluginStatus> PluginStatuses = IPluginManager::Get().QueryStatusForAllPlugins();
+	for (const auto& PluginStatus : PluginStatuses)
+	{
+		if (PluginStatus.bIsEnabled && !PluginStatus.PluginDirectory.IsEmpty())
+		{
+			const FString PluginTemplatesDirectory = FPaths::Combine(*PluginStatus.PluginDirectory, TEXT("Templates"));
+
+			if (IFileManager::Get().DirectoryExists(*PluginTemplatesDirectory))
+			{
+				TemplateRootFolders.Add(PluginTemplatesDirectory);
+			}
+		}
+	}
+
 	// Form a list of all folders that could contain template projects
 	TArray<FString> AllTemplateFolders;
 	for ( auto TemplateRootFolderIt = TemplateRootFolders.CreateConstIterator(); TemplateRootFolderIt; ++TemplateRootFolderIt )
diff --git a/Engine/Source/Runtime/Projects/Private/PluginManager.cpp b/Engine/Source/Runtime/Projects/Private/PluginManager.cpp
--- a/Engine/Source/Runtime/Projects/Private/PluginManager.cpp
+++ b/Engine/Source/Runtime/Projects/Private/PluginManager.cpp
@@ -434,6 +434,7 @@ TArray< FPluginStatus > FPluginManager::QueryStatusForAllPlugins() const
 		PluginStatus.CreatedBy = PluginInfo.CreatedBy;
 		PluginStatus.CreatedByURL = PluginInfo.CreatedByURL;
 		PluginStatus.CategoryPath = PluginInfo.Category;
+		PluginStatus.PluginDirectory = FPaths::GetPath(Plugin->FileName);
 		PluginStatus.bIsEnabled = Plugin->bEnabled;
 		PluginStatus.bIsBuiltIn = ( Plugin->LoadedFrom == EPluginLoadedFrom::Engine );
 		PluginStatus.bIsEnabledByDefault = PluginInfo.bEnabledByDefault;
diff --git a/Engine/Source/Runtime/Projects/Public/Interfaces/IPluginManager.h b/Engine/Source/Runtime/Projects/Public/Interfaces/IPluginManager.h
--- a/Engine/Source/Runtime/Projects/Public/Interfaces/IPluginManager.h
+++ b/Engine/Source/Runtime/Projects/Public/Interfaces/IPluginManager.h
@@ -38,6 +38,9 @@ public:
 	/** Category path (dot-separated list of categories) */
 	FString CategoryPath;
 
+	/** Path to plugin directory on disk */
+	FString PluginDirectory;
+
 	/** True if plugin is currently enabled */
 	bool bIsEnabled;
 
