Add a sort key to template definitions, and use it to position the 2D sidescroller sample next to the 3D version. Don't want to sort templates by display name because they all start with "Code" or "Blueprint", and at least this way we can hack it through data for now.

diff --git a/Engine/Source/Editor/GameProjectGeneration/Classes/TemplateProjectDefs.h b/Engine/Source/Editor/GameProjectGeneration/Classes/TemplateProjectDefs.h
--- a/Engine/Source/Editor/GameProjectGeneration/Classes/TemplateProjectDefs.h
+++ b/Engine/Source/Editor/GameProjectGeneration/Classes/TemplateProjectDefs.h
@@ -72,6 +72,9 @@ class UTemplateProjectDefs : public UObject
 	UPROPERTY(config)
 	TArray<FTemplateReplacement> ReplacementsInFiles;
 
+	UPROPERTY(config)
+	FString SortKey;
+
 	/** Fixes up all strings in this definitions object to replace %TEMPLATENAME% with the supplied template name and %PROJECTNAME% with the supplied project name */
 	void FixupStrings(const FString& TemplateName, const FString& ProjectName);
 
diff --git a/Engine/Source/Editor/GameProjectGeneration/Private/SNewProjectWizard.cpp b/Engine/Source/Editor/GameProjectGeneration/Private/SNewProjectWizard.cpp
--- a/Engine/Source/Editor/GameProjectGeneration/Private/SNewProjectWizard.cpp
+++ b/Engine/Source/Editor/GameProjectGeneration/Private/SNewProjectWizard.cpp
@@ -16,19 +16,29 @@ struct FTemplateItem
 {
 	FText Name;
 	FText Description;
+	FString SortKey;
 	FString ProjectFile;
 	TSharedPtr<FSlateBrush> TemplateThumbnail;
 	bool bGenerateCode;
 
-	FTemplateItem(const FText& InName, const FText& InDescription, const TSharedPtr<FSlateBrush>& InTemplateThumbnail, bool InGenerateCode, const FString& InProjectFile)
+	FTemplateItem(const FText& InName, const FText& InDescription, const TSharedPtr<FSlateBrush>& InTemplateThumbnail, bool InGenerateCode, const FString& InSortKey, const FString& InProjectFile)
 		: Name(InName)
 		, Description(InDescription)
+		, SortKey(InSortKey)
 		, ProjectFile(InProjectFile)
 		, TemplateThumbnail(InTemplateThumbnail)
 		, bGenerateCode(InGenerateCode)
 	{}
 };
 
+struct FTemplateSortPredicate
+{
+	bool operator()(const TSharedPtr<FTemplateItem>& A, const TSharedPtr<FTemplateItem>& B) const
+	{
+		return A->SortKey < B->SortKey;
+	}
+};
+
 BEGIN_SLATE_FUNCTION_BUILD_OPTIMIZATION
 void SNewProjectWizard::Construct( const FArguments& InArgs )
 {
@@ -891,6 +901,7 @@ void SNewProjectWizard::FindTemplateProjects()
 		LOCTEXT("BlankProjectDescription", "A clean empty project with no code."),
 		MakeShareable( new FSlateBrush( *FEditorStyle::GetBrush("GameProjectDialog.BlankProjectThumbnail") ) ),
 		false,				// bGenerateCode
+		TEXT("_1"),			// SortKey
 		TEXT("")			// No filename, this is a generation template
 		)) );
 
@@ -899,6 +910,7 @@ void SNewProjectWizard::FindTemplateProjects()
 		LOCTEXT("BasicCodeProjectDescription", "An empty project with some basic game framework code classes created."),
 		MakeShareable( new FSlateBrush( *FEditorStyle::GetBrush("GameProjectDialog.BasicCodeThumbnail") ) ),
 		true,				// bGenerateCode
+		TEXT("_2"),			// SortKey
 		TEXT("")			// No filename, this is a generation template
 		)) );
 
@@ -975,14 +987,18 @@ void SNewProjectWizard::FindTemplateProjects()
 						DynamicBrush = MakeShareable( new FSlateDynamicImageBrush(BrushName , FVector2D(128,128) ) );
 					}
 
-					if (FPaths::GetCleanFilename(ProjectFilename) == GameProjectUtils::GetDefaultProjectTemplateFilename())
+					// Get the sort key
+					FString SortKey = TemplateDefs->SortKey;
+					if(SortKey.Len() == 0)
 					{
-						TemplateList.Insert(MakeShareable(new FTemplateItem(TemplateName, TemplateDescription, DynamicBrush, bGenerateCode, ProjectFilename)), 0);
+						SortKey = FPaths::GetCleanFilename(ProjectFilename);
 					}
-					else
+					if (FPaths::GetCleanFilename(ProjectFilename) == GameProjectUtils::GetDefaultProjectTemplateFilename())
 					{
-						TemplateList.Add(MakeShareable(new FTemplateItem(TemplateName, TemplateDescription, DynamicBrush, bGenerateCode, ProjectFilename)));
+						SortKey = TEXT("_0");
 					}
+
+					TemplateList.Add(MakeShareable(new FTemplateItem(TemplateName, TemplateDescription, DynamicBrush, bGenerateCode, SortKey, ProjectFilename)));
 				}
 			}
 			else
@@ -992,6 +1008,9 @@ void SNewProjectWizard::FindTemplateProjects()
 			}
 		}
 	}
+
+	// Sort the template folders
+	TemplateList.Sort(FTemplateSortPredicate());
 }
 
 
