Add link to code from Add Code to Project window

diff --git a/Engine/Source/Editor/GameProjectGeneration/Private/SNewClassDialog.cpp b/Engine/Source/Editor/GameProjectGeneration/Private/SNewClassDialog.cpp
--- a/Engine/Source/Editor/GameProjectGeneration/Private/SNewClassDialog.cpp
+++ b/Engine/Source/Editor/GameProjectGeneration/Private/SNewClassDialog.cpp
@@ -209,21 +209,58 @@ void SNewClassDialog::Construct( const FArguments& InArgs )
 						// Class label
 						+SHorizontalBox::Slot()
 						.AutoWidth()
-						.VAlign(VAlign_Center)
-						.Padding(0, 0, 12, 0)
 						[
-							SNew(STextBlock)
-							.TextStyle( FEditorStyle::Get(), "NewClassDialog.SelectedParentClassLabel" )
-							.Text( LOCTEXT( "ParentClassLabel", "Selected Class" ) )
+							SNew(SVerticalBox)
+
+							+SVerticalBox::Slot()
+							.AutoHeight()
+							.VAlign(VAlign_Center)
+							.Padding(0, 0, 12, 0)
+							[
+								SNew(STextBlock)
+								.TextStyle( FEditorStyle::Get(), "NewClassDialog.SelectedParentClassLabel" )
+								.Text( LOCTEXT( "ParentClassLabel", "Selected Class" ) )
+							]
+
+							+SVerticalBox::Slot()
+							.AutoHeight()
+							.VAlign(VAlign_Center)
+							.Padding(0, 0, 12, 0)
+							[
+								SNew(STextBlock)
+								.TextStyle(FEditorStyle::Get(), "NewClassDialog.SelectedParentClassLabel")
+								.Text(LOCTEXT("ParentClassSourceLabel", "Selected Class Source"))
+							]
 						]
 
 						// Class selection preview
 						+SHorizontalBox::Slot()
-						.FillWidth(1.f)
-						.VAlign(VAlign_Center)
 						[
-							SNew(STextBlock)
-							.Text( this, &SNewClassDialog::GetSelectedParentClassName )
+							SNew(SVerticalBox)
+
+							+SVerticalBox::Slot()
+							.AutoHeight()
+							.VAlign(VAlign_Center)
+							.Padding(0, 0, 12, 0)
+							[
+								SNew(STextBlock)
+								.Text( this, &SNewClassDialog::GetSelectedParentClassName )
+							]
+
+							+ SVerticalBox::Slot()
+							.AutoHeight()
+							.VAlign(VAlign_Bottom)
+							.HAlign(HAlign_Left)
+							.Padding(0.0f, 0.0f, 0.0f, 0.0f)
+							[
+								SNew(SHyperlink)
+								.Style(FCoreStyle::Get(), "Hyperlink")
+								.TextStyle(FEditorStyle::Get(), "DetailsView.GoToCodeHyperlinkStyle")
+								.OnNavigate(this, &SNewClassDialog::OnEditCodeClicked)
+								.Text(this, &SNewClassDialog::GetSelectedParentClassFilename)
+								.ToolTipText(LOCTEXT("GoToCode_ToolTip", "Click to open this source file in a text editor"))
+								.Visibility(this, &SNewClassDialog::GetSourceHyperlinkVisibility)
+							]
 						]
 					]
 				]
@@ -573,6 +610,7 @@ TSharedRef<ITableRow> SNewClassDialog::MakeParentClassListViewWidget(TSharedPtr<
 	return
 		SNew( STableRow<TSharedPtr<FParentClassItem>>, OwnerTable )
 		.Style(FEditorStyle::Get(), "NewClassDialog.ParentClassListView.TableRow")
+		.ToolTip(IDocumentation::Get()->CreateToolTip(FText::FromString(ClassDescription), nullptr, "Shared/Classes", ClassName))
 		[
 			SNew(SBox).HeightOverride(ItemHeight)
 			[
@@ -620,6 +658,45 @@ FString SNewClassDialog::GetSelectedParentClassName() const
 	return ParentClassInfo.IsSet() ? ParentClassInfo.GetClassName() : TEXT("");
 }
 
+FString GetClassHeaderPath(const UClass* Class)
+{
+	if (Class)
+	{
+		FString ClassHeaderPath;
+		if (FSourceCodeNavigation::FindClassHeaderPath(Class, ClassHeaderPath) && IFileManager::Get().FileSize(*ClassHeaderPath) != INDEX_NONE)
+		{
+			return ClassHeaderPath;
+		}
+	}
+	return FString();
+}
+
+EVisibility SNewClassDialog::GetSourceHyperlinkVisibility() const
+{
+	return (GetClassHeaderPath(ParentClassInfo.BaseClass).Len() > 0 ? EVisibility::Visible : EVisibility::Hidden);
+}
+
+FString SNewClassDialog::GetSelectedParentClassFilename() const
+{
+	const FString ClassHeaderPath = GetClassHeaderPath(ParentClassInfo.BaseClass);
+	if (ClassHeaderPath.Len() > 0)
+	{
+		return FPaths::GetCleanFilename(*ClassHeaderPath);
+	}
+	return FString();
+}
+
+void SNewClassDialog::OnEditCodeClicked()
+{
+	const FString ClassHeaderPath = GetClassHeaderPath(ParentClassInfo.BaseClass);
+	if (ClassHeaderPath.Len() > 0)
+	{
+		const FString AbsoluteHeaderPath = IFileManager::Get().ConvertToAbsolutePathForExternalAppForRead(*ClassHeaderPath);
+		FSourceCodeNavigation::OpenSourceFile(AbsoluteHeaderPath);
+	}
+}
+
+
 void SNewClassDialog::OnParentClassItemDoubleClicked( TSharedPtr<FParentClassItem> TemplateItem )
 {
 	// Advance to the name page
diff --git a/Engine/Source/Editor/GameProjectGeneration/Private/SNewClassDialog.h b/Engine/Source/Editor/GameProjectGeneration/Private/SNewClassDialog.h
--- a/Engine/Source/Editor/GameProjectGeneration/Private/SNewClassDialog.h
+++ b/Engine/Source/Editor/GameProjectGeneration/Private/SNewClassDialog.h
@@ -32,6 +32,15 @@ private:
 	/** Gets the currently selected parent class name */
 	FString GetSelectedParentClassName() const;
 
+	/** Gets the currently selected parent class's filename */
+	FString GetSelectedParentClassFilename() const;
+
+	/** Whether the hyper link to go to source should be visible */
+	EVisibility GetSourceHyperlinkVisibility() const;
+
+	/** Handler for when the seleted parent class's filename is clicked */
+	void OnEditCodeClicked();
+
 	/** Handler for when a parent class item is double clicked */
 	void OnParentClassItemDoubleClicked( TSharedPtr<FParentClassItem> TemplateItem );
 
