'/' and '\\\\' are correctly not valid in project name at all as was intended.
Reversed the priority of the validity checks to give more consistent results.
Fixes TTP #341288

diff --git a/Engine/Source/Editor/GameProjectGeneration/Private/SNewProjectWizard.cpp b/Engine/Source/Editor/GameProjectGeneration/Private/SNewProjectWizard.cpp
--- a/Engine/Source/Editor/GameProjectGeneration/Private/SNewProjectWizard.cpp
+++ b/Engine/Source/Editor/GameProjectGeneration/Private/SNewProjectWizard.cpp
@@ -1108,32 +1108,34 @@ void SNewProjectWizard::UpdateProjectFileValidity( )
 
 	// Name and Location Validity
 	{
-		if ( CurrentProjectFileName.Contains(TEXT("/")) || CurrentProjectFileName.Contains(TEXT("\\")) )
+		bLastNameAndLocationValidityCheckSuccessful = true;
+
+		if ( !FPlatformMisc::IsValidAbsolutePathFormat(CurrentProjectFilePath) )
 		{
 			bLastNameAndLocationValidityCheckSuccessful = false;
-			LastNameAndLocationValidityErrorText = LOCTEXT( "SlashOrBackslashInProjectName", "The project name may not contain a slash or backslash" );
+			LastNameAndLocationValidityErrorText = LOCTEXT( "InvalidFolderPath", "The folder path is invalid" );
 		}
 		else
 		{
 			FText FailReason;
-			bLastNameAndLocationValidityCheckSuccessful = GameProjectUtils::IsValidProjectFileForCreation( GetProjectFilenameWithPath(), FailReason );
-			if ( !bLastNameAndLocationValidityCheckSuccessful )
+			if ( !GameProjectUtils::IsValidProjectFileForCreation(GetProjectFilenameWithPath(), FailReason) )
 			{
+				bLastNameAndLocationValidityCheckSuccessful = false;
 				LastNameAndLocationValidityErrorText = FailReason;
 			}
 		}
 
-		if ( !FPlatformMisc::IsValidAbsolutePathFormat(CurrentProjectFilePath) )
+		if ( CurrentProjectFileName.Contains(TEXT("/")) || CurrentProjectFileName.Contains(TEXT("\\")) )
 		{
 			bLastNameAndLocationValidityCheckSuccessful = false;
-			LastNameAndLocationValidityErrorText = LOCTEXT( "InvalidFolderPath", "The folder path is invalid" );
+			LastNameAndLocationValidityErrorText = LOCTEXT("SlashOrBackslashInProjectName", "The project name may not contain a slash or backslash");
 		}
 		else
 		{
 			FText FailReason;
-			bLastNameAndLocationValidityCheckSuccessful = GameProjectUtils::IsValidProjectFileForCreation( GetProjectFilenameWithPath(), FailReason );
-			if ( !bLastNameAndLocationValidityCheckSuccessful )
+			if ( !GameProjectUtils::IsValidProjectFileForCreation(GetProjectFilenameWithPath(), FailReason) )
 			{
+				bLastNameAndLocationValidityCheckSuccessful = false;
 				LastNameAndLocationValidityErrorText = FailReason;
 			}
 		}
