UE4: One more non-unity fix.

diff --git a/Engine/Source/Editor/GameProjectGeneration/Private/SNewClassDialog.cpp b/Engine/Source/Editor/GameProjectGeneration/Private/SNewClassDialog.cpp
--- a/Engine/Source/Editor/GameProjectGeneration/Private/SNewClassDialog.cpp
+++ b/Engine/Source/Editor/GameProjectGeneration/Private/SNewClassDialog.cpp
@@ -7,6 +7,7 @@
 #include "ClassViewerFilter.h"
 #include "Editor/ClassViewer/Private/SClassViewer.h"
 #include "DesktopPlatformModule.h"
+#include "Editor/Documentation/Public/IDocumentation.h"
 
 #define LOCTEXT_NAMESPACE "GameProjectGeneration"
 
