UE4: Build fix. Since GameProjectGeneration now depends on the Documentation module, it is now in the DynamicallyLoadedModuleNames list.

diff --git a/Engine/Source/Editor/GameProjectGeneration/GameProjectGeneration.Build.cs b/Engine/Source/Editor/GameProjectGeneration/GameProjectGeneration.Build.cs
--- a/Engine/Source/Editor/GameProjectGeneration/GameProjectGeneration.Build.cs
+++ b/Engine/Source/Editor/GameProjectGeneration/GameProjectGeneration.Build.cs
@@ -38,6 +38,7 @@ public class GameProjectGeneration : ModuleRules
             new string[] {
 				"ClassViewer",
                 "DesktopPlatform",
+                "Documentation",
                 "MainFrame",
             }
 		);
