EditoronlyBP enabled in new Blank project.

diff --git a/Engine/Source/Editor/GameProjectGeneration/Private/GameProjectUtils.cpp b/Engine/Source/Editor/GameProjectGeneration/Private/GameProjectUtils.cpp
--- a/Engine/Source/Editor/GameProjectGeneration/Private/GameProjectUtils.cpp
+++ b/Engine/Source/Editor/GameProjectGeneration/Private/GameProjectUtils.cpp
@@ -1478,6 +1478,26 @@ bool GameProjectUtils::GenerateConfigFiles(const FString& NewProjectPath, const
 		}
 	}
 
+	// DefaultEditor.ini
+	{
+		const FString DefaultEditorIniFilename = ProjectConfigPath / TEXT("DefaultEditor.ini");
+		FString FileContents;
+		FileContents += TEXT("[EditoronlyBP]") LINE_TERMINATOR;
+		FileContents += TEXT("bAllowClassAndBlueprintPinMatching=true") LINE_TERMINATOR;
+		FileContents += TEXT("bReplaceBlueprintWithClass=true") LINE_TERMINATOR;
+		FileContents += TEXT("bDontLoadBlueprintOutsideEditor=true") LINE_TERMINATOR;
+		FileContents += TEXT("bBlueprintIsNotBlueprintType=true") LINE_TERMINATOR;
+
+		if (WriteOutputFile(DefaultEditorIniFilename, FileContents, OutFailReason))
+		{
+			OutCreatedFiles.Add(DefaultEditorIniFilename);
+		}
+		else
+		{
+			return false;
+		}
+	}
+
 	return true;
 }
 
