Provide a proper error message if a file cannot be deleted during project upgrade, and fix offering to checkout the project file if necessary.

diff --git a/Engine/Plugins/Editor/PluginsEditor/Source/PluginsEditor/Private/SPluginListTile.cpp b/Engine/Plugins/Editor/PluginsEditor/Source/PluginsEditor/Private/SPluginListTile.cpp
--- a/Engine/Plugins/Editor/PluginsEditor/Source/PluginsEditor/Private/SPluginListTile.cpp
+++ b/Engine/Plugins/Editor/PluginsEditor/Source/PluginsEditor/Private/SPluginListTile.cpp
@@ -277,7 +277,7 @@ void SPluginListTile::OnEnablePluginCheckboxChanged(ESlateCheckBoxState::Type Ne
 
 	ItemData->PluginStatus.bIsEnabled = bNewEnabledState;
 
-	FGameProjectGenerationModule::Get().TryMakeProjectFileWriteable();
+	FGameProjectGenerationModule::Get().TryMakeProjectFileWriteable(FPaths::GetProjectFilePath());
 
 	FText FailMessage;
 	if(!IProjectManager::Get().SetPluginEnabled(ItemData->PluginStatus.Name, bNewEnabledState, FailMessage))
diff --git a/Engine/Source/Developer/DesktopPlatform/Private/DesktopPlatformBase.cpp b/Engine/Source/Developer/DesktopPlatform/Private/DesktopPlatformBase.cpp
--- a/Engine/Source/Developer/DesktopPlatform/Private/DesktopPlatformBase.cpp
+++ b/Engine/Source/Developer/DesktopPlatform/Private/DesktopPlatformBase.cpp
@@ -303,7 +303,7 @@ bool FDesktopPlatformBase::OpenProject(const FString& ProjectFileName)
 	return true;
 }
 
-bool FDesktopPlatformBase::CleanGameProject(const FString &ProjectDir, FFeedbackContext* Warn)
+bool FDesktopPlatformBase::CleanGameProject(const FString& ProjectDir, FString& OutFailPath, FFeedbackContext* Warn)
 {
 	// Begin a task
 	Warn->BeginSlowTask(LOCTEXT("CleaningProject", "Removing stale build products..."), true);
@@ -319,7 +319,8 @@ bool FDesktopPlatformBase::CleanGameProject(const FString &ProjectDir, FFeedback
 		// Remove the file
 		if(!IFileManager::Get().Delete(*FileNames[Idx]))
 		{
-			Warn->Logf(ELogVerbosity::Error, TEXT("ERROR: Couldn't delete file '%s'"), *FileNames[Idx]);
+			OutFailPath = FileNames[Idx];
+			Warn->EndSlowTask();
 			return false;
 		}
 
@@ -333,7 +334,8 @@ bool FDesktopPlatformBase::CleanGameProject(const FString &ProjectDir, FFeedback
 		// Remove the directory
 		if(!IFileManager::Get().DeleteDirectory(*DirectoryNames[Idx], false, true))
 		{
-			Warn->Logf(ELogVerbosity::Error, TEXT("ERROR: Couldn't delete directory '%s'"), *DirectoryNames[Idx]);
+			OutFailPath = DirectoryNames[Idx];
+			Warn->EndSlowTask();
 			return false;
 		}
 
diff --git a/Engine/Source/Developer/DesktopPlatform/Private/DesktopPlatformBase.h b/Engine/Source/Developer/DesktopPlatform/Private/DesktopPlatformBase.h
--- a/Engine/Source/Developer/DesktopPlatform/Private/DesktopPlatformBase.h
+++ b/Engine/Source/Developer/DesktopPlatform/Private/DesktopPlatformBase.h
@@ -35,7 +35,7 @@ public:
 
 	virtual bool OpenProject(const FString& ProjectFileName) override;
 
-	virtual bool CleanGameProject(const FString& ProjectDir, FFeedbackContext* Warn) override;
+	virtual bool CleanGameProject(const FString& ProjectDir, FString& OutFailPath, FFeedbackContext* Warn) override;
 	virtual bool CompileGameProject(const FString& RootDir, const FString& ProjectFileName, FFeedbackContext* Warn) override;
 	virtual bool GenerateProjectFiles(const FString& RootDir, const FString& ProjectFileName, FFeedbackContext* Warn) override;
 
diff --git a/Engine/Source/Developer/DesktopPlatform/Public/IDesktopPlatform.h b/Engine/Source/Developer/DesktopPlatform/Public/IDesktopPlatform.h
--- a/Engine/Source/Developer/DesktopPlatform/Public/IDesktopPlatform.h
+++ b/Engine/Source/Developer/DesktopPlatform/Public/IDesktopPlatform.h
@@ -300,7 +300,7 @@ public:
 	* @param ProjectDirName		Directory for the project
 	* @param OutFileNames		Output array of the project's build products
 	*/
-	virtual bool CleanGameProject(const FString& ProjectDir, FFeedbackContext* Warn) = 0;
+	virtual bool CleanGameProject(const FString& ProjectDir, FString& OutFailPath, FFeedbackContext* Warn) = 0;
 
 	/**
 	* Compiles a game project.
diff --git a/Engine/Source/Editor/GameProjectGeneration/Private/GameProjectGenerationModule.cpp b/Engine/Source/Editor/GameProjectGeneration/Private/GameProjectGenerationModule.cpp
--- a/Engine/Source/Editor/GameProjectGeneration/Private/GameProjectGenerationModule.cpp
+++ b/Engine/Source/Editor/GameProjectGeneration/Private/GameProjectGenerationModule.cpp
@@ -43,9 +43,9 @@ void FGameProjectGenerationModule::OpenAddCodeToProjectDialog()
 	AddCodeToProjectDialogOpenedEvent.Broadcast();
 }
 
-void FGameProjectGenerationModule::TryMakeProjectFileWriteable()
+void FGameProjectGenerationModule::TryMakeProjectFileWriteable(const FString& ProjectFile)
 {
-	GameProjectUtils::TryMakeProjectFileWriteable();
+	GameProjectUtils::TryMakeProjectFileWriteable(ProjectFile);
 }
 
 void FGameProjectGenerationModule::CheckForOutOfDateGameProjectFile()
@@ -54,9 +54,9 @@ void FGameProjectGenerationModule::CheckForOutOfDateGameProjectFile()
 }
 
 
-bool FGameProjectGenerationModule::UpdateGameProject(const FString& EngineIdentifier)
+bool FGameProjectGenerationModule::UpdateGameProject(const FString& ProjectFile, const FString& EngineIdentifier, FText& OutFailReason)
 {
-	return GameProjectUtils::UpdateGameProject(EngineIdentifier);
+	return GameProjectUtils::UpdateGameProject(ProjectFile, EngineIdentifier, OutFailReason);
 }
 
 
diff --git a/Engine/Source/Editor/GameProjectGeneration/Private/GameProjectUtils.cpp b/Engine/Source/Editor/GameProjectGeneration/Private/GameProjectUtils.cpp
--- a/Engine/Source/Editor/GameProjectGeneration/Private/GameProjectUtils.cpp
+++ b/Engine/Source/Editor/GameProjectGeneration/Private/GameProjectUtils.cpp
@@ -691,25 +691,9 @@ void GameProjectUtils::OnWarningReasonOk()
 	}
 }
 
-bool GameProjectUtils::UpdateGameProject(const FString& EngineIdentifier)
+bool GameProjectUtils::UpdateGameProject(const FString& ProjectFile, const FString& EngineIdentifier, FText& OutFailReason)
 {
-	const FString& ProjectFilename = FPaths::IsProjectFilePathSet() ? FPaths::GetProjectFilePath() : FString();
-	if ( !ProjectFilename.IsEmpty() )
-	{
-		FText FailReason;
-		bool bWasCheckedOut = false;
-		if ( !UpdateGameProjectFile(ProjectFilename, EngineIdentifier, NULL, bWasCheckedOut, FailReason) )
-		{
-			// The user chose to update, but the update failed. Notify the user.
-			UE_LOG(LogGameProjectGeneration, Error, TEXT("%s failed to update. %s"), *ProjectFilename, *FailReason.ToString() );
-			return false;
-		}
-
-		// The project was updated successfully.
-		UE_LOG(LogGameProjectGeneration, Log, TEXT("%s was successfully updated."), *ProjectFilename );
-	}
-
-	return true;
+	return UpdateGameProjectFile(ProjectFile, EngineIdentifier, NULL, OutFailReason);
 }
 
 void GameProjectUtils::OpenAddCodeToProjectDialog()
@@ -1241,17 +1225,36 @@ bool GameProjectUtils::CreateProjectFromTemplate(const FString& NewProjectFile,
 
 	// Generate the project file
 	{
-		FText LocalFailReason;
-		if (IProjectManager::Get().DuplicateProjectFile(TemplateFile, NewProjectFile, FDesktopPlatformModule::Get()->GetCurrentEngineIdentifier(), LocalFailReason))
+		// Load the source project
+		FProjectDescriptor Project;
+		if(!Project.Load(TemplateFile, OutFailReason))
 		{
-			CreatedFiles.Add(NewProjectFile);
+			DeleteCreatedFiles(DestFolder, CreatedFiles);
+			return false;
 		}
-		else
+
+		// Update it to current
+		Project.EngineAssociation = FDesktopPlatformModule::Get()->GetCurrentEngineIdentifier();
+		Project.EpicSampleNameHash = 0;
+
+		// Fix up module names
+		const FString BaseSourceName = FPaths::GetBaseFilename(TemplateFile);
+		const FString BaseNewName = FPaths::GetBaseFilename(NewProjectFile);
+		for ( auto ModuleIt = Project.Modules.CreateIterator(); ModuleIt; ++ModuleIt )
+		{
+			FModuleDescriptor& ModuleInfo = *ModuleIt;
+			ModuleInfo.Name = FName(*ModuleInfo.Name.ToString().Replace(*BaseSourceName, *BaseNewName));
+		}
+
+		// Save it to disk
+		if(!Project.Save(NewProjectFile, OutFailReason))
 		{
-			OutFailReason = LocalFailReason;
 			DeleteCreatedFiles(DestFolder, CreatedFiles);
 			return false;
 		}
+
+		// Add it to the list of created files
+		CreatedFiles.Add(NewProjectFile);
 	}
 
 	if ( bShouldGenerateCode )
@@ -2442,17 +2445,12 @@ void GameProjectUtils::UpdateProject(const TArray<FString>* StartupModuleNames)
 	FText FailReason;
 	FText UpdateMessage;
 	SNotificationItem::ECompletionState NewCompletionState;
-	bool bWasCheckedOut = false;
-	if ( UpdateGameProjectFile(ProjectFilename, FDesktopPlatformModule::Get()->GetCurrentEngineIdentifier(), StartupModuleNames, bWasCheckedOut, FailReason) )
+	if ( UpdateGameProjectFile(ProjectFilename, FDesktopPlatformModule::Get()->GetCurrentEngineIdentifier(), StartupModuleNames, FailReason) )
 	{
 		// The project was updated successfully.
 		FFormatNamedArguments Args;
 		Args.Add( TEXT("ShortFilename"), FText::FromString( ShortFilename ) );
 		UpdateMessage = FText::Format( LOCTEXT("ProjectFileUpdateComplete", "{ShortFilename} was successfully updated."), Args );
-		if ( bWasCheckedOut )
-		{
-			UpdateMessage = FText::Format( LOCTEXT("ProjectFileUpdateCheckin", "{ShortFilename} was successfully updated. Please check this file into source control."), Args );
-		}
 		NewCompletionState = SNotificationItem::CS_Success;
 	}
 	else
@@ -2484,44 +2482,56 @@ void GameProjectUtils::OnUpdateProjectCancel()
 	}
 }
 
-void GameProjectUtils::TryMakeProjectFileWriteable()
+void GameProjectUtils::TryMakeProjectFileWriteable(const FString& ProjectFile)
 {
-	FString ProjectFileName = FPaths::GetProjectFilePath();
-
 	// First attempt to check out the file if SCC is enabled
 	if ( ISourceControlModule::Get().IsEnabled() )
 	{
 		FText FailReason;
-		GameProjectUtils::CheckoutGameProjectFile(ProjectFileName, FailReason);
+		GameProjectUtils::CheckoutGameProjectFile(ProjectFile, FailReason);
 	}
 
-	// Check if it's writeable
-	if(FPlatformFileManager::Get().GetPlatformFile().IsReadOnly(*ProjectFileName))
+	// Check if it's writable
+	if(FPlatformFileManager::Get().GetPlatformFile().IsReadOnly(*ProjectFile))
 	{
-		FText ShouldMakeProjectWriteable = LOCTEXT("ShouldMakeProjectWriteable_Message", "'{ProjectFilename}' is read-only and cannot be updated, would you like to make it writeable?");
+		FText ShouldMakeProjectWriteable = LOCTEXT("ShouldMakeProjectWriteable_Message", "'{ProjectFilename}' is read-only and cannot be updated. Would you like to make it writeable?");
 
 		FFormatNamedArguments Arguments;
-		Arguments.Add( TEXT("ProjectFilename"), FText::FromString(ProjectFileName));
+		Arguments.Add( TEXT("ProjectFilename"), FText::FromString(ProjectFile));
 
 		if(FMessageDialog::Open(EAppMsgType::YesNo, FText::Format(ShouldMakeProjectWriteable, Arguments)) == EAppReturnType::Yes)
 		{
-			FPlatformFileManager::Get().GetPlatformFile().SetReadOnly(*ProjectFileName, false);
+			FPlatformFileManager::Get().GetPlatformFile().SetReadOnly(*ProjectFile, false);
 		}
 	}
 }
 
-bool GameProjectUtils::UpdateGameProjectFile(const FString& ProjectFilename, const FString& EngineIdentifier, const TArray<FString>* StartupModuleNames, bool& OutbWasCheckedOut, FText& OutFailReason)
+bool GameProjectUtils::UpdateGameProjectFile(const FString& ProjectFile, const FString& EngineIdentifier, const TArray<FString>* StartupModuleNames, FText& OutFailReason)
 {
 	// Make sure we can write to the project file
-	TryMakeProjectFileWriteable();
+	TryMakeProjectFileWriteable(ProjectFile);
 
-	// Now tell the project manager to update the file
-	if (!IProjectManager::Get().UpdateLoadedProjectFileToCurrent(StartupModuleNames, EngineIdentifier, OutFailReason))
+	// Load the descriptor
+	FProjectDescriptor Descriptor;
+	if(Descriptor.Load(ProjectFile, OutFailReason))
 	{
-		return false;
-	}
+		// Freshen version information
+		Descriptor.EngineAssociation = EngineIdentifier;
 
-	return true;
+		// Replace the modules names, if specified
+		if(StartupModuleNames != NULL)
+		{
+			Descriptor.Modules.Empty();
+			for(int32 Idx = 0; Idx < StartupModuleNames->Num(); Idx++)
+			{
+				Descriptor.Modules.Add(FModuleDescriptor(*(*StartupModuleNames)[Idx]));
+			}
+		}
+
+		// Update file on disk
+		return Descriptor.Save(ProjectFile, OutFailReason);
+	}
+	return false;
 }
 
 bool GameProjectUtils::CheckoutGameProjectFile(const FString& ProjectFilename, FText& OutFailReason)
diff --git a/Engine/Source/Editor/GameProjectGeneration/Private/GameProjectUtils.h b/Engine/Source/Editor/GameProjectGeneration/Private/GameProjectUtils.h
--- a/Engine/Source/Editor/GameProjectGeneration/Private/GameProjectUtils.h
+++ b/Engine/Source/Editor/GameProjectGeneration/Private/GameProjectUtils.h
@@ -142,10 +142,10 @@ public:
 	static void CheckAndWarnProjectFilenameValid();
 
 	/** Checks out the current project file (or prompts to make writable) */
-	static void TryMakeProjectFileWriteable();
+	static void TryMakeProjectFileWriteable(const FString& ProjectFile);
 
-	/** Updates the currently loaded project. Returns true if the project was updated successfully or if no update was needed */
-	static bool UpdateGameProject(const FString &EngineIdentifier);
+	/** Updates the given project file to an engine identifier. Returns true if the project was updated successfully or if no update was needed */
+	static bool UpdateGameProject(const FString& ProjectFile, const FString& EngineIdentifier, FText& OutFailReason);
 
 	/** Opens a dialog to add code files to a project */
 	static void OpenAddCodeToProjectDialog();
@@ -354,16 +354,16 @@ private:
 	static void OnUpdateProjectCancel();
 
 	/**
-	 * Updates the loaded game project file to the current version
+	 * Updates the loaded game project file to the current version and to use the given modules
 	 *
 	 * @param ProjectFilename		The name of the project (used to checkout from source control)
+	 * @param EngineIdentifier		The identifier for the engine to open the project with
 	 * @param StartupModuleNames	if specified, replaces the existing module names with this version
-	 * @param OutbWasCheckedOut		Out, whether or not the project was checked out from source control
 	 * @param OutFailReason			Out, if unsuccessful this is the reason why
 
 	 * @return true, if successful
 	 */
-	static bool UpdateGameProjectFile(const FString& ProjectFilename, const FString& EngineIdentifier, const TArray<FString>* StartupModuleNames, bool& OutbWasCheckedOut, FText& OutFailReason);
+	static bool UpdateGameProjectFile(const FString& ProjectFilename, const FString& EngineIdentifier, const TArray<FString>* StartupModuleNames, FText& OutFailReason);
 
 	/** Checks the specified game project file out from source control */
 	static bool CheckoutGameProjectFile(const FString& ProjectFilename, FText& OutFailReason);
diff --git a/Engine/Source/Editor/GameProjectGeneration/Private/SProjectBrowser.cpp b/Engine/Source/Editor/GameProjectGeneration/Private/SProjectBrowser.cpp
--- a/Engine/Source/Editor/GameProjectGeneration/Private/SProjectBrowser.cpp
+++ b/Engine/Source/Editor/GameProjectGeneration/Private/SProjectBrowser.cpp
@@ -959,9 +959,11 @@ bool SProjectBrowser::OpenProject( const FString& InProjectFile )
 		}
 		if(Selection == OpenExistingButton)
 		{
-			if(!FDesktopPlatformModule::Get()->CleanGameProject(FPaths::GetPath(ProjectFile), GWarn))
+			FString FailPath;
+			if(!FDesktopPlatformModule::Get()->CleanGameProject(FPaths::GetPath(ProjectFile), FailPath, GWarn))
 			{
-				FMessageDialog::Open( EAppMsgType::Ok, LOCTEXT("ConvertProjectCleanFailed", "Couldn't clean project build files. Delete the project intermediate folder and try again.") );
+				FText FailMessage = FText::Format(LOCTEXT("ConvertProjectCleanFailed", "{0} could not be removed. Try deleting it manually and try again."), FText::FromString(FailPath));
+				FMessageDialog::Open(EAppMsgType::Ok, FailMessage);
 				return false;
 			}
 		}
@@ -1019,9 +1021,9 @@ bool SProjectBrowser::OpenProject( const FString& InProjectFile )
 			}
 
 			// Update the game project to the latest version. This will prompt to check-out as necessary. We don't need to write the engine identifier directly, because it won't use the right .uprojectdirs logic.
-			if(!GameProjectUtils::UpdateGameProject(TEXT("")) || !FDesktopPlatformModule::Get()->SetEngineIdentifierForProject(ProjectFile, CurrentIdentifier))
+			if(!GameProjectUtils::UpdateGameProject(ProjectFile, CurrentIdentifier, FailReason))
 			{
-				if(FMessageDialog::Open(EAppMsgType::YesNo, LOCTEXT("ProjectUpgradeFailure", "Project file could not be updated to latest version. Attempt to open anyway?")) != EAppReturnType::Yes)
+				if(FMessageDialog::Open(EAppMsgType::YesNo, LOCTEXT("ProjectUpgradeFailure", "The project file could not be updated to latest version. Attempt to open anyway?")) != EAppReturnType::Yes)
 				{
 					return false;
 				}
diff --git a/Engine/Source/Editor/GameProjectGeneration/Public/GameProjectGenerationModule.h b/Engine/Source/Editor/GameProjectGeneration/Public/GameProjectGenerationModule.h
--- a/Engine/Source/Editor/GameProjectGeneration/Public/GameProjectGenerationModule.h
+++ b/Engine/Source/Editor/GameProjectGeneration/Public/GameProjectGenerationModule.h
@@ -48,13 +48,13 @@ public:
 	FAddCodeToProjectDialogOpenedEvent& OnAddCodeToProjectDialogOpened() { return AddCodeToProjectDialogOpenedEvent; }
 
 	/** Tries to make the project file writable. Prompts to check out as necessary. */
-	virtual void TryMakeProjectFileWriteable();
+	virtual void TryMakeProjectFileWriteable(const FString& ProjectFile);
 
 	/** Prompts the user to update his project file, if necessary. */
 	virtual void CheckForOutOfDateGameProjectFile();
 
 	/** Updates the currently loaded project. Returns true if the project was updated successfully or if no update was needed */
-	virtual bool UpdateGameProject(const FString& EngineIdentifier);
+	virtual bool UpdateGameProject(const FString& ProjectFile, const FString& EngineIdentifier, FText& OutFailReason);
 
 	/** Updates the current code project */
 	virtual bool UpdateCodeProject(FText& OutFailReason);
diff --git a/Engine/Source/Editor/UnrealEd/Private/Commandlets/UpdateGameProjectCommandlet.cpp b/Engine/Source/Editor/UnrealEd/Private/Commandlets/UpdateGameProjectCommandlet.cpp
--- a/Engine/Source/Editor/UnrealEd/Private/Commandlets/UpdateGameProjectCommandlet.cpp
+++ b/Engine/Source/Editor/UnrealEd/Private/Commandlets/UpdateGameProjectCommandlet.cpp
@@ -71,9 +71,10 @@ int32 UUpdateGameProjectCommandlet::Main( const FString& InParams )
 
 	UE_LOG(LogUpdateGameProjectCommandlet, Display, TEXT("Updating project file %s to %s..."), *ProjectFilePath, *EngineIdentifier);
 
-	if ( !FGameProjectGenerationModule::Get().UpdateGameProject(EngineIdentifier) )
+	FText FailReason;
+	if ( !FGameProjectGenerationModule::Get().UpdateGameProject(ProjectFilePath, EngineIdentifier, FailReason) )
 	{
-		// UpdateGameProject produces it's own error to the log
+		UE_LOG(LogUpdateGameProjectCommandlet, Error, TEXT("Couldn't update game project: %s"), *FailReason.ToString());
 		return 1;
 	}
 
diff --git a/Engine/Source/Runtime/Projects/Private/ProjectManager.cpp b/Engine/Source/Runtime/Projects/Private/ProjectManager.cpp
--- a/Engine/Source/Runtime/Projects/Private/ProjectManager.cpp
+++ b/Engine/Source/Runtime/Projects/Private/ProjectManager.cpp
@@ -126,56 +126,6 @@ bool FProjectManager::GenerateNewProjectFile(const FString& NewProjectFilename,
 	return Descriptor.Save(NewProjectFilename, OutFailReason);
 }
 
-bool FProjectManager::DuplicateProjectFile(const FString& SourceProjectFilename, const FString& NewProjectFilename, const FString& EngineIdentifier, FText& OutFailReason)
-{
-	// Load the source project
-	FProjectDescriptor Project;
-	if(!Project.Load(SourceProjectFilename, OutFailReason))
-	{
-		return false;
-	}
-
-	// Update it to current
-	Project.EngineAssociation = EngineIdentifier;
-	Project.EpicSampleNameHash = 0;
-
-	// Fix up module names
-	const FString BaseSourceName = FPaths::GetBaseFilename(SourceProjectFilename);
-	const FString BaseNewName = FPaths::GetBaseFilename(NewProjectFilename);
-	for ( auto ModuleIt = Project.Modules.CreateIterator(); ModuleIt; ++ModuleIt )
-	{
-		FModuleDescriptor& ModuleInfo = *ModuleIt;
-		ModuleInfo.Name = FName(*ModuleInfo.Name.ToString().Replace(*BaseSourceName, *BaseNewName));
-	}
-
-	// Save it to disk
-	return Project.Save(NewProjectFilename, OutFailReason);
-}
-
-bool FProjectManager::UpdateLoadedProjectFileToCurrent(const TArray<FString>* StartupModuleNames, const FString& EngineIdentifier, FText& OutFailReason)
-{
-	if ( !CurrentProject.IsValid() )
-	{
-		return false;
-	}
-
-	// Freshen version information
-	CurrentProject->EngineAssociation = EngineIdentifier;
-
-	// Replace the modules names, if specified
-	if(StartupModuleNames != NULL)
-	{
-		CurrentProject->Modules.Empty();
-		for(int32 Idx = 0; Idx < StartupModuleNames->Num(); Idx++)
-		{
-			CurrentProject->Modules.Add(FModuleDescriptor(*(*StartupModuleNames)[Idx]));
-		}
-	}
-
-	// Update file on disk
-	return CurrentProject->Save(FPaths::GetProjectFilePath(), OutFailReason);
-}
-
 bool FProjectManager::SignSampleProject(const FString& FilePath, const FString& Category, FText& OutFailReason)
 {
 	FProjectDescriptor Descriptor;
diff --git a/Engine/Source/Runtime/Projects/Private/ProjectManager.h b/Engine/Source/Runtime/Projects/Private/ProjectManager.h
--- a/Engine/Source/Runtime/Projects/Private/ProjectManager.h
+++ b/Engine/Source/Runtime/Projects/Private/ProjectManager.h
@@ -20,8 +20,6 @@ public:
 	virtual const FString& GetAutoLoadProjectFileName() override;
 	virtual const FString& NonStaticGetProjectFileExtension() override;
 	virtual bool GenerateNewProjectFile(const FString& NewProjectFilename, const TArray<FString>& ModuleNames, const FString& EngineIdentifier, FText& OutFailReason) override;
-	virtual bool DuplicateProjectFile(const FString& SourceProjectFilename, const FString& NewProjectFilename, const FString& EngineIdentifier, FText& OutFailReason) override;
-	virtual bool UpdateLoadedProjectFileToCurrent(const TArray<FString>* StartupModuleNames, const FString& EngineIdentifier, FText& OutFailReason) override;
 	virtual bool SignSampleProject(const FString& FilePath, const FString& Category, FText& OutFailReason) override;
 	virtual bool QueryStatusForProject(const FString& FilePath, FProjectStatus& OutProjectStatus) const override;
 	virtual bool QueryStatusForCurrentProject(FProjectStatus& OutProjectStatus) const override;
diff --git a/Engine/Source/Runtime/Projects/Public/Interfaces/IProjectManager.h b/Engine/Source/Runtime/Projects/Public/Interfaces/IProjectManager.h
--- a/Engine/Source/Runtime/Projects/Public/Interfaces/IProjectManager.h
+++ b/Engine/Source/Runtime/Projects/Public/Interfaces/IProjectManager.h
@@ -138,29 +138,6 @@ public:
 	 */
 	virtual bool GenerateNewProjectFile(const FString& NewProjectFilename, const TArray<FString>& StartupModuleNames, const FString& EngineIdentifier, FText& OutFailReason) = 0;
 
-	/**
-	 * Generates a new project file based on a source file and saves it to disk at the specified location
-	 * 
-	 * @param SourceProjectFilename	The filename of the file to duplicate
-	 * @param NewProjectFilename	The filename of the file to be written out
-	 * @param EngineIdentifier		Identifier of the engine that the project should be associated with
-	 * @param OutFailReason			When returning false, this provides a display reason why the file could not be created.
-	 *
-	 * @return true when the file was successfully written
-	 */
-	virtual bool DuplicateProjectFile(const FString& SourceProjectFilename, const FString& NewProjectFilename, const FString& EngineIdentifier, FText& OutFailReason) = 0;
-
-	/**
-	 * Updates the loaded game project file to the current version and saves it to disk
-	 * 
-	 * @param StartupModuleNames	When specified, replaces the existing module names with these new ones			
-	 * @param EngineIdentifier		Identifier for the current engine
-	 * @param OutFailReason			When returning false, this provides a display reason why the file could not be created.
-	 *
-	 * @return true if the file was successfully updated and saved to disk
-	 */
-	virtual bool UpdateLoadedProjectFileToCurrent(const TArray<FString>* StartupModuleNames, const FString& EngineIdentifier, FText& OutFailReason) = 0;
-
 	/**
 	 * Sets the project's EpicSampleNameHash (based on its filename) and category, then saves the file to disk.
 	 * This marks the project as a sample and fixes its filename so that
diff --git a/Engine/Source/Runtime/Projects/Public/ModuleDescriptor.h b/Engine/Source/Runtime/Projects/Public/ModuleDescriptor.h
--- a/Engine/Source/Runtime/Projects/Public/ModuleDescriptor.h
+++ b/Engine/Source/Runtime/Projects/Public/ModuleDescriptor.h
@@ -41,7 +41,7 @@ namespace ELoadingPhase
 	 * @param	The string to convert to a value
 	 * @return	The corresponding value, or 'Max' if the string is not valid.
 	 */
-	ELoadingPhase::Type FromString( const TCHAR *Text );
+	PROJECTS_API ELoadingPhase::Type FromString( const TCHAR *Text );
 
 	/**
 	 * Returns the name of a module load phase.
@@ -49,7 +49,7 @@ namespace ELoadingPhase
 	 * @param	The value to convert to a string
 	 * @return	The string representation of this enum value
 	 */
-	const TCHAR* ToString( const ELoadingPhase::Type Value );
+	PROJECTS_API const TCHAR* ToString( const ELoadingPhase::Type Value );
 };
 
 /**
@@ -77,7 +77,7 @@ namespace EHostType
 	 * @param	The string to convert to a value
 	 * @return	The corresponding value, or 'Max' if the string is not valid.
 	 */
-	EHostType::Type FromString( const TCHAR *Text );
+	PROJECTS_API EHostType::Type FromString( const TCHAR *Text );
 
 	/**
 	 * Converts an EHostType::Type value to a string literal
@@ -85,13 +85,13 @@ namespace EHostType
 	 * @param	The value to convert to a string
 	 * @return	The string representation of this enum value
 	 */
-	const TCHAR* ToString( const EHostType::Type Value );
+	PROJECTS_API const TCHAR* ToString( const EHostType::Type Value );
 };
 
 /**
  * Description of a loadable module.
  */
-struct FModuleDescriptor
+struct PROJECTS_API FModuleDescriptor
 {
 	/** Name of this module */
 	FName Name;
