The class wizard now displays the modules listed in your .uproject file as available targets for the new class

TTP# 342920 - EDITOR: Cloned C++ projects do not compile correctly
TTP# 342762 - Editor: Feature Request: Relax the restrictions on what game modules can have new classes created within them

This has allowed the validation logic to be simplified, as it pushes the responsibility onto the user to say which module they want their new class to go into, rather than relying on the validation logic to correctly infer which module the class should be going into.

This also relaxes the previous naming restrictions due to assumptions about module names, as you're now able to add code to any module listed in your .uproject file.

I've tested:
 - Adding code to a normal project.
 - Adding code to a cloned project.
 - Creating a new code based project.
 - Adding code to an empty/blueprint based project.

All of these cases generated code which compiled correctly.

diff --git a/Engine/Source/Editor/GameProjectGeneration/Private/GameProjectUtils.cpp b/Engine/Source/Editor/GameProjectGeneration/Private/GameProjectUtils.cpp
--- a/Engine/Source/Editor/GameProjectGeneration/Private/GameProjectUtils.cpp
+++ b/Engine/Source/Editor/GameProjectGeneration/Private/GameProjectUtils.cpp
@@ -770,7 +770,7 @@ bool GameProjectUtils::IsValidClassNameForCreation(const FString& NewClassName,
 	// Look for a duplicate class on disk in their project
 	{
 		FString UnusedFoundPath;
-		if ( FindSourceFileInProject(NewClassName + ".h", UnusedFoundPath, ModuleInfo) )
+		if ( FindSourceFileInProject(NewClassName + ".h", ModuleInfo.ModuleSourcePath, UnusedFoundPath) )
 		{
 			FFormatNamedArguments Args;
 			Args.Add( TEXT("NewClassName"), FText::FromString( NewClassName ) );
@@ -782,9 +782,9 @@ bool GameProjectUtils::IsValidClassNameForCreation(const FString& NewClassName,
 	return true;
 }
 
-bool GameProjectUtils::AddCodeToProject(const FString& NewClassName, const FString& NewClassPath, const FNewClassInfo ParentClassInfo, FString& OutHeaderFilePath, FString& OutCppFilePath, FText& OutFailReason)
+bool GameProjectUtils::AddCodeToProject(const FString& NewClassName, const FString& NewClassPath, const FModuleContextInfo& ModuleInfo, const FNewClassInfo ParentClassInfo, FString& OutHeaderFilePath, FString& OutCppFilePath, FText& OutFailReason)
 {
-	const bool bAddCodeSuccessful = AddCodeToProject_Internal(NewClassName, NewClassPath, ParentClassInfo, OutHeaderFilePath, OutCppFilePath, OutFailReason);
+	const bool bAddCodeSuccessful = AddCodeToProject_Internal(NewClassName, NewClassPath, ModuleInfo, ParentClassInfo, OutHeaderFilePath, OutCppFilePath, OutFailReason);
 
 	if( FEngineAnalytics::IsAvailable() )
 	{
@@ -1605,8 +1605,9 @@ bool GameProjectUtils::GenerateGameFrameworkSourceCode(const FString& NewProject
 
 	// Used to override the code generation validation since the module we're creating isn't the same as the project we currently have loaded
 	FModuleContextInfo NewModuleInfo;
-	NewModuleInfo.ModuleSourcePath = NewProjectSourcePath;
 	NewModuleInfo.ModuleName = NewProjectName;
+	NewModuleInfo.ModuleType = EHostType::Runtime;
+	NewModuleInfo.ModuleSourcePath = FPaths::ConvertRelativePathToFull(GameModulePath / ""); // Ensure trailing /
 
 	// MyGamePlayerController.h
 	{
@@ -1706,45 +1707,50 @@ bool GameProjectUtils::IsStarterContentAvailableForNewProjects()
 	return (StarterContentFiles.Num() > 0);
 }
 
-GameProjectUtils::FModuleContextInfo GameProjectUtils::GetCurrentModuleContextInfo()
+TArray<GameProjectUtils::FModuleContextInfo> GameProjectUtils::GetCurrentProjectModules()
 {
-	FModuleContextInfo ModuleInfo;
+	const FProjectDescriptor* const CurrentProject = IProjectManager::Get().GetCurrentProject();
+	check(CurrentProject);
 
-	ModuleInfo.ModuleSourcePath = FPaths::GameSourceDir();
-	
-	// Assuming the game name is the same as the primary game module name
-	ModuleInfo.ModuleName = FApp::GetGameName();
-
-	return ModuleInfo;
-}
-
-FString GameProjectUtils::GetSourceRootPath(const bool bIncludeModuleName, const FModuleContextInfo& ModuleInfo)
-{
-	FString SourceDir = ModuleInfo.ModuleSourcePath;
+	TArray<FModuleContextInfo> RetModuleInfos;
 
-	if(bIncludeModuleName)
+	if (!GameProjectUtils::ProjectHasCodeFiles() || CurrentProject->Modules.Num() == 0)
 	{
-		SourceDir /= ModuleInfo.ModuleName;
+		// If this project doesn't currently have any code in it, we need to add a dummy entry for the game
+		// so that we can still use the class wizard (this module will be created once we add a class)
+		FModuleContextInfo ModuleInfo;
+		ModuleInfo.ModuleName = FApp::GetGameName();
+		ModuleInfo.ModuleType = EHostType::Runtime;
+		ModuleInfo.ModuleSourcePath = FPaths::ConvertRelativePathToFull(FPaths::GameSourceDir() / ModuleInfo.ModuleName / ""); // Ensure trailing /
+		RetModuleInfos.Emplace(ModuleInfo);
 	}
 
-	SourceDir /= "";
+	// Resolve out the paths for each module and add the cut-down into to our output array
+	for (const FModuleDescriptor& ModuleDesc : CurrentProject->Modules)
+	{
+		FModuleContextInfo ModuleInfo;
+		ModuleInfo.ModuleName = ModuleDesc.Name.ToString();
+		ModuleInfo.ModuleType = ModuleDesc.Type;
 
-	return FPaths::ConvertRelativePathToFull(SourceDir);
-}
+		// Try and find the .Build.cs file for this module within our currently loaded project's Source directory
+		FString TmpPath;
+		if (!FindSourceFileInProject(ModuleInfo.ModuleName + ".Build.cs", FPaths::GameSourceDir(), TmpPath))
+		{
+			continue;
+		}
 
-bool GameProjectUtils::IsValidSourcePath(const FString& InPath, const bool bIncludeModuleName, const FModuleContextInfo& ModuleInfo, FText* const OutFailReason)
-{
-	FString RootPath = GetSourceRootPath(bIncludeModuleName, ModuleInfo);
+		// Chop the .Build.cs file off the end of the path
+		ModuleInfo.ModuleSourcePath = FPaths::GetPath(TmpPath);
+		ModuleInfo.ModuleSourcePath = FPaths::ConvertRelativePathToFull(ModuleInfo.ModuleSourcePath / ""); // Ensure trailing /
 
-	// Only allow partial module name matches if we already have code; the first class added to a project *must* be for the game module
-	const bool bHasCodeFiles = GameProjectUtils::ProjectHasCodeFiles();
-	if(bIncludeModuleName && bHasCodeFiles)
-	{
-		// If we're including the module name, then we want to allow variations of it, eg) MyModule, MyModuleEditor, MyModuleClient
-		// Those variations are valid, so we trim the last / from the path so that the StartsWith check below allows these variations
-		RootPath = RootPath.LeftChop(1); // Trim trailing /
+		RetModuleInfos.Emplace(ModuleInfo);
 	}
 
+	return RetModuleInfos;
+}
+
+bool GameProjectUtils::IsValidSourcePath(const FString& InPath, const FModuleContextInfo& ModuleInfo, FText* const OutFailReason)
+{
 	const FString AbsoluteInPath = FPaths::ConvertRelativePathToFull(InPath) / ""; // Ensure trailing /
 
 	// Validate the path contains no invalid characters
@@ -1753,16 +1759,14 @@ bool GameProjectUtils::IsValidSourcePath(const FString& InPath, const bool bIncl
 		return false;
 	}
 
-	if(!AbsoluteInPath.StartsWith(RootPath))
+	if(!AbsoluteInPath.StartsWith(ModuleInfo.ModuleSourcePath))
 	{
 		if(OutFailReason)
 		{
 			FFormatNamedArguments Args;
-			Args.Add(TEXT("RootSourcePath"), FText::FromString(GetSourceRootPath(bIncludeModuleName, ModuleInfo)));
-			const FText FormatString = (bIncludeModuleName)
-				? LOCTEXT("SourcePathInvalidModuleRoot", "All source code must exist within a valid module folder in the projects source path, e.g. {RootSourcePath}")
-				: LOCTEXT("SourcePathInvalidRoot", "All source code must exist within the projects source path: {RootSourcePath}");
-			*OutFailReason = FText::Format( FormatString, Args );
+			Args.Add(TEXT("ModuleName"), FText::FromString(ModuleInfo.ModuleName));
+			Args.Add(TEXT("RootSourcePath"), FText::FromString(ModuleInfo.ModuleSourcePath));
+			*OutFailReason = FText::Format( LOCTEXT("SourcePathInvalidForModule", "All source code for '{ModuleName}' must exist within '{RootSourcePath}'"), Args );
 		}
 		return false;
 	}
@@ -1770,20 +1774,19 @@ bool GameProjectUtils::IsValidSourcePath(const FString& InPath, const bool bIncl
 	return true;
 }
 
-bool GameProjectUtils::CalculateSourcePaths(const FString& InPath, FString& OutModuleName, FString& OutHeaderPath, FString& OutSourcePath, const FModuleContextInfo& ModuleInfo, FText* const OutFailReason)
+bool GameProjectUtils::CalculateSourcePaths(const FString& InPath, const FModuleContextInfo& ModuleInfo, FString& OutHeaderPath, FString& OutSourcePath, FText* const OutFailReason)
 {
 	const FString AbsoluteInPath = FPaths::ConvertRelativePathToFull(InPath) / ""; // Ensure trailing /
 	OutHeaderPath = AbsoluteInPath;
 	OutSourcePath = AbsoluteInPath;
 
 	EClassLocation ClassPathLocation = EClassLocation::UserDefined;
-	if(!GetClassLocation(InPath, OutModuleName, ClassPathLocation, ModuleInfo, OutFailReason))
+	if(!GetClassLocation(InPath, ModuleInfo, ClassPathLocation, OutFailReason))
 	{
 		return false;
 	}
 
-	const FString BaseRootPath = GetSourceRootPath(false/*bIncludeModuleName*/, ModuleInfo);
-	const FString RootPath = BaseRootPath / OutModuleName / ""; // Ensure trailing /
+	const FString RootPath = ModuleInfo.ModuleSourcePath;
 	const FString PublicPath = RootPath / "Public" / "";		// Ensure trailing /
 	const FString PrivatePath = RootPath / "Private" / "";		// Ensure trailing /
 	const FString ClassesPath = RootPath / "Classes" / "";		// Ensure trailing /
@@ -1830,26 +1833,17 @@ bool GameProjectUtils::CalculateSourcePaths(const FString& InPath, FString& OutM
 	return !OutHeaderPath.IsEmpty() && !OutSourcePath.IsEmpty();
 }
 
-bool GameProjectUtils::GetClassLocation(const FString& InPath, FString& OutModuleName, EClassLocation& OutClassLocation, const FModuleContextInfo& ModuleInfo, FText* const OutFailReason)
+bool GameProjectUtils::GetClassLocation(const FString& InPath, const FModuleContextInfo& ModuleInfo, EClassLocation& OutClassLocation, FText* const OutFailReason)
 {
 	const FString AbsoluteInPath = FPaths::ConvertRelativePathToFull(InPath) / ""; // Ensure trailing /
-	OutModuleName.Empty();
 	OutClassLocation = EClassLocation::UserDefined;
 
-	if(!IsValidSourcePath(InPath, true/*bIncludeModuleName*/, ModuleInfo, OutFailReason))
+	if(!IsValidSourcePath(InPath, ModuleInfo, OutFailReason))
 	{
 		return false;
 	}
 
-	// We've validated that this path includes a partial match for our module (eg, MyModule, MyModuleEditor, MyModuleClient)
-	// so extract the actual name of the module from the path so that we can generate the internal folder names correctly
-	const FString BaseRootPath = GetSourceRootPath(false/*bIncludeModuleName*/, ModuleInfo);
-	const int32 ModuleNameStartIndex = BaseRootPath.Len();
-	const int32 ModuleNameEndIndex = AbsoluteInPath.Find(TEXT("/"), ESearchCase::IgnoreCase, ESearchDir::FromStart, ModuleNameStartIndex);
-	check(ModuleNameEndIndex != INDEX_NONE); // this should never happen since AbsoluteInPath ends in a /, and we verified it started with BaseRootPath in IsValidSourcePath
-	OutModuleName = AbsoluteInPath.Mid(ModuleNameStartIndex, ModuleNameEndIndex - ModuleNameStartIndex);
-
-	const FString RootPath = BaseRootPath / OutModuleName / ""; // Ensure trailing /
+	const FString RootPath = ModuleInfo.ModuleSourcePath;
 	const FString PublicPath = RootPath / "Public" / "";		// Ensure trailing /
 	const FString PrivatePath = RootPath / "Private" / "";		// Ensure trailing /
 	const FString ClassesPath = RootPath / "Classes" / "";		// Ensure trailing /
@@ -2127,14 +2121,13 @@ bool GameProjectUtils::GenerateClassHeaderFile(const FString& NewHeaderFileName,
 
 	FString ModuleAPIMacro;
 	{
-		FString ModuleName;
 		EClassLocation ClassPathLocation = EClassLocation::UserDefined;
-		if ( GetClassLocation(NewHeaderFileName, ModuleName, ClassPathLocation, ModuleInfo) )
+		if ( GetClassLocation(NewHeaderFileName, ModuleInfo, ClassPathLocation) )
 		{
 			// If this class isn't Private, make sure and include the API macro so it can be linked within other modules
 			if ( ClassPathLocation != EClassLocation::Private )
 			{
-				ModuleAPIMacro = ModuleName.ToUpper() + "_API "; // include a trailing space for the template formatting
+				ModuleAPIMacro = ModuleInfo.ModuleName.ToUpper() + "_API "; // include a trailing space for the template formatting
 			}
 		}
 	}
@@ -2189,9 +2182,8 @@ bool GameProjectUtils::GenerateClassCPPFile(const FString& NewCPPFileName, const
 	const FString PrefixedClassName = ClassPrefix + UnPrefixedClassName;
 	const FString PrefixedBaseClassName = ClassPrefix + ParentClassInfo.GetClassNameCPP();
 
-	FString ModuleName;
 	EClassLocation ClassPathLocation = EClassLocation::UserDefined;
-	if ( !GetClassLocation(NewCPPFileName, ModuleName, ClassPathLocation, ModuleInfo, &OutFailReason) )
+	if ( !GetClassLocation(NewCPPFileName, ModuleInfo, ClassPathLocation, &OutFailReason) )
 	{
 		return false;
 	}
@@ -2221,19 +2213,18 @@ bool GameProjectUtils::GenerateClassCPPFile(const FString& NewCPPFileName, const
 
 	// Calculate the correct include path for the module header
 	FString ModuleIncludePath;
-	if(FindSourceFileInProject(ModuleName + ".h", ModuleIncludePath, ModuleInfo))
+	if(FindSourceFileInProject(ModuleInfo.ModuleName + ".h", ModuleInfo.ModuleSourcePath, ModuleIncludePath))
 	{
 		// Work out where the module header is; 
 		// if it's Public then we can include it without any path since all Public and Classes folders are on the include path
 		// if it's located elsewhere, then we'll need to include it relative to the module source root as we can't guarantee 
 		// that other folders are on the include paths
-		FString UnusedModuleName;
 		EClassLocation ModuleLocation;
-		if(GetClassLocation(ModuleIncludePath, UnusedModuleName, ModuleLocation, ModuleInfo))
+		if(GetClassLocation(ModuleIncludePath, ModuleInfo, ModuleLocation))
 		{
 			if(ModuleLocation == EClassLocation::Public || ModuleLocation == EClassLocation::Classes)
 			{
-				ModuleIncludePath = ModuleName + ".h";
+				ModuleIncludePath = ModuleInfo.ModuleName + ".h";
 			}
 			else
 			{
@@ -2242,18 +2233,15 @@ bool GameProjectUtils::GenerateClassCPPFile(const FString& NewCPPFileName, const
 				const FString ClassPath = FPaths::ConvertRelativePathToFull(FPaths::GetPath(NewCPPFileName));
 				if(ModulePath == ClassPath)
 				{
-					ModuleIncludePath = ModuleName + ".h";
+					ModuleIncludePath = ModuleInfo.ModuleName + ".h";
 				}
 				else
 				{
-					// Make the include relative to the source root of the module
-					const FString ModuleSourceRoot = GetSourceRootPath(true/*bIncludeModuleName*/, ModuleInfo);
-					
 					// Updates ModuleIncludePath internally
-					if(!FPaths::MakePathRelativeTo(ModuleIncludePath, *ModuleSourceRoot))
+					if(!FPaths::MakePathRelativeTo(ModuleIncludePath, *ModuleInfo.ModuleSourcePath))
 					{
 						// Failed; just assume we can include it without any relative path
-						ModuleIncludePath = ModuleName + ".h";
+						ModuleIncludePath = ModuleInfo.ModuleName + ".h";
 					}
 				}
 			}
@@ -2261,19 +2249,19 @@ bool GameProjectUtils::GenerateClassCPPFile(const FString& NewCPPFileName, const
 		else
 		{
 			// Failed; just assume we can include it without any relative path
-			ModuleIncludePath = ModuleName + ".h";
+			ModuleIncludePath = ModuleInfo.ModuleName + ".h";
 		}
 	}
 	else
 	{
 		// This could potentially fail when generating new projects if the module file hasn't yet been created; just assume we can include it without any relative path
-		ModuleIncludePath = ModuleName + ".h";
+		ModuleIncludePath = ModuleInfo.ModuleName + ".h";
 	}
 
 	// Not all of these will exist in every class template
 	FString FinalOutput = Template.Replace(TEXT("%COPYRIGHT_LINE%"), *MakeCopyrightLine(), ESearchCase::CaseSensitive);
 	FinalOutput = FinalOutput.Replace(TEXT("%UNPREFIXED_CLASS_NAME%"), *UnPrefixedClassName, ESearchCase::CaseSensitive);
-	FinalOutput = FinalOutput.Replace(TEXT("%MODULE_NAME%"), *ModuleName, ESearchCase::CaseSensitive);
+	FinalOutput = FinalOutput.Replace(TEXT("%MODULE_NAME%"), *ModuleInfo.ModuleName, ESearchCase::CaseSensitive);
 	FinalOutput = FinalOutput.Replace(TEXT("%MODULE_INCLUDE_PATH%"), *ModuleIncludePath, ESearchCase::CaseSensitive);
 	FinalOutput = FinalOutput.Replace(TEXT("%PREFIXED_CLASS_NAME%"), *PrefixedClassName, ESearchCase::CaseSensitive);
 	FinalOutput = FinalOutput.Replace(TEXT("%PROPERTY_OVERRIDES%"), *PropertyOverridesStr, ESearchCase::CaseSensitive);
@@ -2602,7 +2590,7 @@ bool GameProjectUtils::ProjectHasCodeFiles()
 	return GameProjectUtils::GetProjectCodeFileCount() > 0;
 }
 
-bool GameProjectUtils::AddCodeToProject_Internal(const FString& NewClassName, const FString& NewClassPath, const FNewClassInfo ParentClassInfo, FString& OutHeaderFilePath, FString& OutCppFilePath, FText& OutFailReason)
+bool GameProjectUtils::AddCodeToProject_Internal(const FString& NewClassName, const FString& NewClassPath, const FModuleContextInfo& ModuleInfo, const FNewClassInfo ParentClassInfo, FString& OutHeaderFilePath, FString& OutCppFilePath, FText& OutFailReason)
 {
 	if ( !ParentClassInfo.IsSet() )
 	{
@@ -2613,9 +2601,6 @@ bool GameProjectUtils::AddCodeToProject_Internal(const FString& NewClassName, co
 	const FString CleanClassName = ParentClassInfo.GetCleanClassName(NewClassName);
 	const FString FinalClassName = ParentClassInfo.GetFinalClassName(NewClassName);
 
-	// Get the context info for the current project
-	const FModuleContextInfo ModuleInfo = GetCurrentModuleContextInfo();
-
 	if ( !IsValidClassNameForCreation(FinalClassName, ModuleInfo, OutFailReason) )
 	{
 		return false;
@@ -2627,10 +2612,9 @@ bool GameProjectUtils::AddCodeToProject_Internal(const FString& NewClassName, co
 		return false;
 	}
 
-	FString ModuleName;
 	FString NewHeaderPath;
 	FString NewCppPath;
-	if ( !CalculateSourcePaths(NewClassPath, ModuleName, NewHeaderPath, NewCppPath, ModuleInfo, &OutFailReason) )
+	if ( !CalculateSourcePaths(NewClassPath, ModuleInfo, NewHeaderPath, NewCppPath, &OutFailReason) )
 	{
 		return false;
 	}
@@ -2715,11 +2699,11 @@ bool GameProjectUtils::AddCodeToProject_Internal(const FString& NewClassName, co
 	return true;
 }
 
-bool GameProjectUtils::FindSourceFileInProject(const FString& InFilename, FString& OutPath, const FModuleContextInfo& ModuleInfo)
+bool GameProjectUtils::FindSourceFileInProject(const FString& InFilename, const FString& InSearchPath, FString& OutPath)
 {
 	TArray<FString> Filenames;
 	const FString FilenameWidcard = TEXT("*") + InFilename;
-	IFileManager::Get().FindFilesRecursive(Filenames, *ModuleInfo.ModuleSourcePath, *FilenameWidcard, true, false, false);
+	IFileManager::Get().FindFilesRecursive(Filenames, *InSearchPath, *FilenameWidcard, true, false, false);
 	
 	if(Filenames.Num())
 	{
diff --git a/Engine/Source/Editor/GameProjectGeneration/Private/GameProjectUtils.h b/Engine/Source/Editor/GameProjectGeneration/Private/GameProjectUtils.h
--- a/Engine/Source/Editor/GameProjectGeneration/Private/GameProjectUtils.h
+++ b/Engine/Source/Editor/GameProjectGeneration/Private/GameProjectUtils.h
@@ -13,6 +13,9 @@ public:
 
 		/** Name of the module */
 		FString ModuleName;
+
+		/** Type of this module, eg, Runtime, Editor, etc */
+		EHostType::Type ModuleType;
 	};
 
 	/** Where is this class located within the Source folder? */
@@ -154,7 +157,7 @@ public:
 	static bool IsValidClassNameForCreation(const FString& NewClassName, const FModuleContextInfo& ModuleInfo, FText& OutFailReason);
 
 	/** Adds new source code to the project. When returning true, OutSyncFileAndLineNumber will be the the preferred target file to sync in the users code editing IDE, formatted for use with GenericApplication::GotoLineInSource */
-	static bool AddCodeToProject(const FString& NewClassName, const FString& NewClassPath, const FNewClassInfo ParentClassInfo, FString& OutHeaderFilePath, FString& OutCppFilePath, FText& OutFailReason);
+	static bool AddCodeToProject(const FString& NewClassName, const FString& NewClassPath, const FModuleContextInfo& ModuleInfo, const FNewClassInfo ParentClassInfo, FString& OutHeaderFilePath, FString& OutCppFilePath, FText& OutFailReason);
 
 	/** Loads a template project definitions object from the TemplateDefs.ini file in the specified project */
 	static UTemplateProjectDefs* LoadTemplateDefs(const FString& ProjectDirectory);
@@ -175,58 +178,45 @@ public:
 	static bool IsStarterContentAvailableForNewProjects();
 
 	/**
-	 * Get the default module context information based on the current value of FApp
-	 */
-	static FModuleContextInfo GetCurrentModuleContextInfo();
-
-	/** 
-	 * Get the absolute root path under which all project source code must exist
-	 *
-	 * @param	bIncludeModuleName	Whether to include the module name in the root path?
-	 * @param	ModuleInfo			Information about the module being validated; use GetCurrentModuleContextInfo() if you want to validate the currently loaded project
-	 * 
-	 * @return	The root path. Will always be an absolute path ending with a /
+	 * Get the information about any modules referenced in the .uproject file of the currently loaded project
 	 */
-	static FString GetSourceRootPath(const bool bIncludeModuleName, const FModuleContextInfo& ModuleInfo);
+	static TArray<FModuleContextInfo> GetCurrentProjectModules();
 
 	/** 
 	 * Check to see if the given path is a valid place to put source code for this project (exists within the source root path) 
 	 *
 	 * @param	InPath				The path to check
-	 * @param	bIncludeModuleName	Whether to require the module name in the root path? (is really used as a prefix so will allow MyModule, MyModuleEditor, etc)
-	 * @param	ModuleInfo			Information about the module being validated; use GetCurrentModuleContextInfo() if you want to validate the currently loaded project
+	 * @param	ModuleInfo			Information about the module being validated
 	 * @param	OutFailReason		Optional parameter to fill with failure information
 	 * 
 	 * @return	true if the path is valid, false otherwise
 	 */
-	static bool IsValidSourcePath(const FString& InPath, const bool bIncludeModuleName, const FModuleContextInfo& ModuleInfo, FText* const OutFailReason = nullptr);
+	static bool IsValidSourcePath(const FString& InPath, const FModuleContextInfo& ModuleInfo, FText* const OutFailReason = nullptr);
 
 	/** 
 	 * Given the path provided, work out where generated .h and .cpp files would be placed
 	 *
 	 * @param	InPath				The path to use a base
-	 * @param	OutModuleName		The module name extracted from the path (the part after GetSourceRootPath(false))
+	 * @param	ModuleInfo			Information about the module being validated
 	 * @param	OutHeaderPath		The path where the .h file should be placed
 	 * @param	OutSourcePath		The path where the .cpp file should be placed
-	 * @param	ModuleInfo			Information about the module being validated; use GetCurrentModuleContextInfo() if you want to validate the currently loaded project
 	 * @param	OutFailReason		Optional parameter to fill with failure information
 	 * 
 	 * @return	false if the paths are invalid
 	 */
-	static bool CalculateSourcePaths(const FString& InPath, FString& OutModuleName, FString& OutHeaderPath, FString& OutSourcePath, const FModuleContextInfo& ModuleInfo, FText* const OutFailReason = nullptr);
+	static bool CalculateSourcePaths(const FString& InPath, const FModuleContextInfo& ModuleInfo, FString& OutHeaderPath, FString& OutSourcePath, FText* const OutFailReason = nullptr);
 
 	/** 
 	 * Given the path provided, work out where it's located within the Source folder
 	 *
 	 * @param	InPath				The path to use a base
-	 * @param	OutModuleName		The module name extracted from the path (the part after GetSourceRootPath(false))
+	 * @param	ModuleInfo			Information about the module being validated
 	 * @param	OutClassLocation	The location within the Source folder
-	 * @param	ModuleInfo			Information about the module being validated; use GetCurrentModuleContextInfo() if you want to validate the currently loaded project
 	 * @param	OutFailReason		Optional parameter to fill with failure information
 	 * 
 	 * @return	false if the paths are invalid
 	 */
-	static bool GetClassLocation(const FString& InPath, FString& OutModuleName, EClassLocation& OutClassLocation, const FModuleContextInfo& ModuleInfo, FText* const OutFailReason = nullptr);
+	static bool GetClassLocation(const FString& InPath, const FModuleContextInfo& ModuleInfo, EClassLocation& OutClassLocation, FText* const OutFailReason = nullptr);
 
 	/** Creates a copy of a project directory in order to upgrade it. */
 	static bool DuplicateProjectForUpgrade( const FString& InProjectFile, FString &OutNewProjectFile );
@@ -372,13 +362,13 @@ private:
 	static bool ProjectHasCodeFiles();
 
 	/** Internal handler for AddCodeToProject*/
-	static bool AddCodeToProject_Internal(const FString& NewClassName, const FString& NewClassPath, const FNewClassInfo ParentClassInfo, FString& OutHeaderFilePath, FString& OutCppFilePath, FText& OutFailReason);
+	static bool AddCodeToProject_Internal(const FString& NewClassName, const FString& NewClassPath, const FModuleContextInfo& ModuleInfo, const FNewClassInfo ParentClassInfo, FString& OutHeaderFilePath, FString& OutCppFilePath, FText& OutFailReason);
 
 	/** Handler for the user confirming they've read the name legnth warning */
 	static void OnWarningReasonOk();
 
 	/** Given a source file name, find its location within the project */
-	static bool FindSourceFileInProject(const FString& InFilename, FString& OutPath, const FModuleContextInfo& ModuleInfo);
+	static bool FindSourceFileInProject(const FString& InFilename, const FString& InSearchPath, FString& OutPath);
 
 private:
 	static TWeakPtr<SNotificationItem> UpdateGameProjectNotification;
diff --git a/Engine/Source/Editor/GameProjectGeneration/Private/SNewClassDialog.cpp b/Engine/Source/Editor/GameProjectGeneration/Private/SNewClassDialog.cpp
--- a/Engine/Source/Editor/GameProjectGeneration/Private/SNewClassDialog.cpp
+++ b/Engine/Source/Editor/GameProjectGeneration/Private/SNewClassDialog.cpp
@@ -24,6 +24,11 @@ struct FParentClassItem
 class FNativeClassParentFilter : public IClassViewerFilter
 {
 public:
+	FNativeClassParentFilter(const TSharedPtr<GameProjectUtils::FModuleContextInfo>* InSelectedModuleInfoPtr)
+		: SelectedModuleInfoPtr(InSelectedModuleInfoPtr)
+	{
+	}
+
 	virtual bool IsClassAllowed(const FClassViewerInitializationOptions& InInitOptions, const UClass* InClass, TSharedRef< FClassViewerFilterFuncs > InFilterFuncs ) override
 	{
 		// You may not make native classes based on blueprint generated classes
@@ -33,9 +38,9 @@ public:
 		const bool bIsExplicitlyUObject = (InClass == UObject::StaticClass());
 
 		// Is this class in the same module as our current module?
-		const GameProjectUtils::FModuleContextInfo ModuleInfo = GameProjectUtils::GetCurrentModuleContextInfo();
 		const FString ClassModuleName = InClass->GetOutermost()->GetName().RightChop( FString(TEXT("/Script/")).Len() );
-		const bool bIsInDestinationModule = (ModuleInfo.ModuleName == ClassModuleName);
+		const TSharedPtr<GameProjectUtils::FModuleContextInfo>& ModuleInfo = *SelectedModuleInfoPtr;
+		const bool bIsInDestinationModule = (ModuleInfo.IsValid() && ModuleInfo->ModuleName == ClassModuleName);
 
 		// You need API if you are either not UObject itself and you are not in the destination module
 		const bool bNeedsAPI = !bIsExplicitlyUObject && !bIsInDestinationModule;
@@ -54,14 +59,53 @@ public:
 	{
 		return false;
 	}
+
+private:
+	/** Pointer to the currently selected module in the new class dialog */
+	const TSharedPtr<GameProjectUtils::FModuleContextInfo>* SelectedModuleInfoPtr;
 };
 
 BEGIN_SLATE_FUNCTION_BUILD_OPTIMIZATION
 void SNewClassDialog::Construct( const FArguments& InArgs )
 {
-	ModuleInfo = GameProjectUtils::GetCurrentModuleContextInfo();
+	{
+		TArray<GameProjectUtils::FModuleContextInfo> CurrentModules = GameProjectUtils::GetCurrentProjectModules();
+		check(CurrentModules.Num()); // this should never happen since GetCurrentProjectModules is supposed to add a dummy runtime module if the project currently has no modules
+
+		AvailableModules.Reserve(CurrentModules.Num());
+		for(const GameProjectUtils::FModuleContextInfo& ModuleInfo : CurrentModules)
+		{
+			AvailableModules.Emplace(MakeShareable(new GameProjectUtils::FModuleContextInfo(ModuleInfo)));
+		}
+	}
+
+	// If we have a runtime module with the same name as our project, then use that
+	// Otherwise, set out default target module as the first runtime module in the list
+	{
+		const FString ProjectName = FApp::GetGameName();
+		for(const auto& AvailableModule : AvailableModules)
+		{
+			if(AvailableModule->ModuleName == ProjectName)
+			{
+				SelectedModuleInfo = AvailableModule;
+				break;
+			}
+			
+			if(AvailableModule->ModuleType == EHostType::Runtime)
+			{
+				SelectedModuleInfo = AvailableModule;
+				// keep going in case we find a better match
+			}
+		}
+
+		if (!SelectedModuleInfo.IsValid())
+		{
+			// No runtime modules? Just take the first available module then
+			SelectedModuleInfo = AvailableModules[0];
+		}
+	}
 
-	NewClassPath = GameProjectUtils::GetSourceRootPath(true/*bIncludeModuleName*/, ModuleInfo);
+	NewClassPath = SelectedModuleInfo->ModuleSourcePath;
 	ClassLocation = GameProjectUtils::EClassLocation::UserDefined; // the first call to UpdateInputValidity will set this correctly based on NewClassPath
 
 	ParentClassInfo = GameProjectUtils::FNewClassInfo(InArgs._Class);
@@ -87,7 +131,7 @@ void SNewClassDialog::Construct( const FArguments& InArgs )
 	Options.bShowObjectRootClass = true;
 
 	// Prevent creating native classes based on blueprint classes
-	Options.ClassFilter = MakeShareable(new FNativeClassParentFilter);
+	Options.ClassFilter = MakeShareable(new FNativeClassParentFilter(&SelectedModuleInfo));
 
 	ClassViewer = StaticCastSharedRef<SClassViewer>(FModuleManager::LoadModuleChecked<FClassViewerModule>("ClassViewer").CreateClassViewer(Options, FOnClassPicked::CreateSP(this, &SNewClassDialog::OnAdvancedClassSelected)));
 
@@ -394,6 +438,22 @@ void SNewClassDialog::Construct( const FArguments& InArgs )
 												.OnTextChanged( this, &SNewClassDialog::OnClassNameTextChanged )
 											]
 
+											+SHorizontalBox::Slot()
+											.AutoWidth()
+											.Padding(6.0f, 0.0f, 0.0f, 0.0f)
+											[
+												SAssignNew(AvailableModulesCombo, SComboBox<TSharedPtr<GameProjectUtils::FModuleContextInfo>>)
+												.ToolTipText( LOCTEXT("ModuleComboToolTip", "Choose the target module for your new class") )
+												.OptionsSource( &AvailableModules )
+												.InitiallySelectedItem( SelectedModuleInfo )
+												.OnSelectionChanged( this, &SNewClassDialog::SelectedModuleComboBoxSelectionChanged )
+												.OnGenerateWidget( this, &SNewClassDialog::MakeWidgetForSelectedModuleCombo )
+												[
+													SNew(STextBlock)
+													.Text( this, &SNewClassDialog::GetSelectedModuleComboText )
+												]
+											]
+
 											+SHorizontalBox::Slot()
 											.AutoWidth()
 											.Padding(6.0f, 0.0f, 0.0f, 0.0f)
@@ -858,6 +918,18 @@ FText SNewClassDialog::OnGetClassPathText() const
 void SNewClassDialog::OnClassPathTextChanged(const FText& NewText)
 {
 	NewClassPath = NewText.ToString();
+
+	// If the user has selected a path which matches the root of a known module, then update our selected module to be that module
+	for(const auto& AvailableModule : AvailableModules)
+	{
+		if(NewClassPath.StartsWith(AvailableModule->ModuleSourcePath))
+		{
+			SelectedModuleInfo = AvailableModule;
+			AvailableModulesCombo->SetSelectedItem(SelectedModuleInfo);
+			break;
+		}
+	}
+
 	UpdateInputValidity();
 }
 
@@ -889,7 +961,7 @@ void SNewClassDialog::FinishClicked()
 	FString CppFilePath;
 
 	FText FailReason;
-	if ( GameProjectUtils::AddCodeToProject(NewClassName, NewClassPath, ParentClassInfo, HeaderFilePath, CppFilePath, FailReason) )
+	if ( GameProjectUtils::AddCodeToProject(NewClassName, NewClassPath, *SelectedModuleInfo, ParentClassInfo, HeaderFilePath, CppFilePath, FailReason) )
 	{
 		// Prevent periodic validity checks. This is to prevent a brief error message about the class already existing while you are exiting.
 		bPreventPeriodicValidityChecksUntilNextChange = true;
@@ -956,6 +1028,18 @@ FReply SNewClassDialog::HandleChooseFolderButtonClicked()
 			}
 
 			NewClassPath = FolderName;
+
+			// If the user has selected a path which matches the root of a known module, then update our selected module to be that module
+			for(const auto& AvailableModule : AvailableModules)
+			{
+				if(NewClassPath.StartsWith(AvailableModule->ModuleSourcePath))
+				{
+					SelectedModuleInfo = AvailableModule;
+					AvailableModulesCombo->SetSelectedItem(SelectedModuleInfo);
+					break;
+				}
+			}
+
 			UpdateInputValidity();
 		}
 	}
@@ -963,6 +1047,40 @@ FReply SNewClassDialog::HandleChooseFolderButtonClicked()
 	return FReply::Handled();
 }
 
+FText SNewClassDialog::GetSelectedModuleComboText() const
+{
+	FFormatNamedArguments Args;
+	Args.Add(TEXT("ModuleName"), FText::FromString(SelectedModuleInfo->ModuleName));
+	Args.Add(TEXT("ModuleType"), FText::FromString(EHostType::ToString(SelectedModuleInfo->ModuleType)));
+	return FText::Format(LOCTEXT("ModuleComboEntry", "{ModuleName} ({ModuleType})"), Args);
+}
+
+void SNewClassDialog::SelectedModuleComboBoxSelectionChanged(TSharedPtr<GameProjectUtils::FModuleContextInfo> Value, ESelectInfo::Type SelectInfo)
+{
+	const FString& OldModulePath = SelectedModuleInfo->ModuleSourcePath;
+	const FString& NewModulePath = Value->ModuleSourcePath;
+
+	SelectedModuleInfo = Value;
+
+	// Update the class path to be rooted to the new module location
+	const FString AbsoluteClassPath = FPaths::ConvertRelativePathToFull(NewClassPath) / ""; // Ensure trailing /
+	if(AbsoluteClassPath.StartsWith(OldModulePath))
+	{
+		NewClassPath = AbsoluteClassPath.Replace(*OldModulePath, *NewModulePath);
+	}
+
+	UpdateInputValidity();
+}
+
+TSharedRef<SWidget> SNewClassDialog::MakeWidgetForSelectedModuleCombo(TSharedPtr<GameProjectUtils::FModuleContextInfo> Value)
+{
+	FFormatNamedArguments Args;
+	Args.Add(TEXT("ModuleName"), FText::FromString(Value->ModuleName));
+	Args.Add(TEXT("ModuleType"), FText::FromString(EHostType::ToString(Value->ModuleType)));
+	return SNew(STextBlock)
+		.Text(FText::Format(LOCTEXT("ModuleComboEntry", "{ModuleName} ({ModuleType})"), Args));
+}
+
 FSlateColor SNewClassDialog::GetClassLocationTextColor(GameProjectUtils::EClassLocation InLocation) const
 {
 	return (ClassLocation == InLocation) ? FSlateColor(FLinearColor(0, 0, 0)) : FSlateColor(FLinearColor(0.72f, 0.72f, 0.72f, 1.f));
@@ -999,12 +1117,10 @@ void SNewClassDialog::OnClassLocationChanged(ESlateCheckBoxState::Type InChecked
 	{
 		const FString AbsoluteClassPath = FPaths::ConvertRelativePathToFull(NewClassPath) / ""; // Ensure trailing /
 
-		FString ModuleName;
 		GameProjectUtils::EClassLocation TmpClassLocation = GameProjectUtils::EClassLocation::UserDefined;
-		GameProjectUtils::GetClassLocation(AbsoluteClassPath, ModuleName, TmpClassLocation, ModuleInfo);
+		GameProjectUtils::GetClassLocation(AbsoluteClassPath, *SelectedModuleInfo, TmpClassLocation);
 
-		const FString BaseRootPath = GameProjectUtils::GetSourceRootPath(false/*bIncludeModuleName*/, ModuleInfo);
-		const FString RootPath = BaseRootPath / ModuleName / "";	// Ensure trailing /
+		const FString RootPath = SelectedModuleInfo->ModuleSourcePath;
 		const FString PublicPath = RootPath / "Public" / "";		// Ensure trailing /
 		const FString PrivatePath = RootPath / "Private" / "";		// Ensure trailing /
 
@@ -1052,15 +1168,14 @@ void SNewClassDialog::UpdateInputValidity()
 	bLastInputValidityCheckSuccessful = true;
 
 	// Validate the path first since this has the side effect of updating the UI
-	FString ModuleName;
-	bLastInputValidityCheckSuccessful = GameProjectUtils::CalculateSourcePaths(NewClassPath, ModuleName, CalculatedClassHeaderName, CalculatedClassSourceName, ModuleInfo, &LastInputValidityErrorText);
+	bLastInputValidityCheckSuccessful = GameProjectUtils::CalculateSourcePaths(NewClassPath, *SelectedModuleInfo, CalculatedClassHeaderName, CalculatedClassSourceName, &LastInputValidityErrorText);
 	CalculatedClassHeaderName /= ParentClassInfo.GetHeaderFilename(NewClassName);
 	CalculatedClassSourceName /= ParentClassInfo.GetSourceFilename(NewClassName);
 
 	// If the source paths check as succeeded, check to see if we're using a Public/Private class
 	if(bLastInputValidityCheckSuccessful)
 	{
-		GameProjectUtils::GetClassLocation(NewClassPath, ModuleName, ClassLocation, ModuleInfo);
+		GameProjectUtils::GetClassLocation(NewClassPath, *SelectedModuleInfo, ClassLocation);
 
 		// We only care about the Public and Private folders
 		if(ClassLocation != GameProjectUtils::EClassLocation::Public && ClassLocation != GameProjectUtils::EClassLocation::Private)
@@ -1076,7 +1191,7 @@ void SNewClassDialog::UpdateInputValidity()
 	// Validate the class name only if the path is valid
 	if ( bLastInputValidityCheckSuccessful )
 	{
-		bLastInputValidityCheckSuccessful = GameProjectUtils::IsValidClassNameForCreation(NewClassName, ModuleInfo, LastInputValidityErrorText);
+		bLastInputValidityCheckSuccessful = GameProjectUtils::IsValidClassNameForCreation(NewClassName, *SelectedModuleInfo, LastInputValidityErrorText);
 	}
 
 	LastPeriodicValidityCheckTime = FSlateApplication::Get().GetCurrentTime();
diff --git a/Engine/Source/Editor/GameProjectGeneration/Private/SNewClassDialog.h b/Engine/Source/Editor/GameProjectGeneration/Private/SNewClassDialog.h
--- a/Engine/Source/Editor/GameProjectGeneration/Private/SNewClassDialog.h
+++ b/Engine/Source/Editor/GameProjectGeneration/Private/SNewClassDialog.h
@@ -122,6 +122,15 @@ private:
 	/** Handler for when the "Choose Folder" button is clicked */
 	FReply HandleChooseFolderButtonClicked( );
 
+	/** Get the combo box text for the currently selected module */
+	FText GetSelectedModuleComboText() const;
+
+	/** Called when the currently selected module is changed */
+	void SelectedModuleComboBoxSelectionChanged(TSharedPtr<GameProjectUtils::FModuleContextInfo> Value, ESelectInfo::Type SelectInfo);
+
+	/** Create the widget to use as the combo box entry for the given module info */
+	TSharedRef<SWidget> MakeWidgetForSelectedModuleCombo(TSharedPtr<GameProjectUtils::FModuleContextInfo> Value);
+
 private:
 
 	/** Get the text color to use for the given class location checkbox */
@@ -166,6 +175,9 @@ private:
 	/** The editable text box to enter the current name */
 	TSharedPtr<SEditableTextBox> ClassNameEditBox;
 
+	/** The available modules combo box */
+	TSharedPtr<SComboBox<TSharedPtr<GameProjectUtils::FModuleContextInfo>>> AvailableModulesCombo;
+
 	/** The name of the class being created */
 	FString NewClassName;
 
@@ -205,6 +217,9 @@ private:
 	/** Whether the class should be created as a Public or Private class */
 	GameProjectUtils::EClassLocation ClassLocation;
 
-	/** Information about the currently loaded module; used for class validation */
-	GameProjectUtils::FModuleContextInfo ModuleInfo;
+	/** Information about the currently available modules for this project */
+	TArray<TSharedPtr<GameProjectUtils::FModuleContextInfo>> AvailableModules;
+
+	/** Information about the currently selected module; used for class validation */
+	TSharedPtr<GameProjectUtils::FModuleContextInfo> SelectedModuleInfo;
 };
