Allow C++ plugins to hook game project generation. Contributed by solid-angle, https://github.com/EpicGames/UnrealEngine/pull/336.

diff --git a/Engine/Source/Editor/GameProjectGeneration/Classes/DefaultTemplateProjectDefs.h b/Engine/Source/Editor/GameProjectGeneration/Classes/DefaultTemplateProjectDefs.h
new file mode 100644
--- /dev/null
+++ b/Engine/Source/Editor/GameProjectGeneration/Classes/DefaultTemplateProjectDefs.h
@@ -0,0 +1,18 @@
+// Copyright 1998-2014 Epic Games, Inc. All Rights Reserved.
+
+#pragma once
+
+#include "TemplateProjectDefs.h"
+#include "DefaultTemplateProjectDefs.generated.h"
+
+UCLASS()
+class GAMEPROJECTGENERATION_API UDefaultTemplateProjectDefs : public UTemplateProjectDefs
+{
+	GENERATED_UCLASS_BODY()
+
+	virtual bool GeneratesCode(const FString& ProjectTemplatePath) const override;
+
+	virtual bool IsClassRename(const FString& DestFilename, const FString& SrcFilename, const FString& FileExtension) const override;
+
+	virtual void AddConfigValues(TArray<FTemplateConfigValue>& ConfigValuesToSet, const FString& TemplateName, const FString& ProjectName, bool bShouldGenerateCode) const override;
+};
diff --git a/Engine/Source/Editor/GameProjectGeneration/Classes/TemplateProjectDefs.h b/Engine/Source/Editor/GameProjectGeneration/Classes/TemplateProjectDefs.h
--- a/Engine/Source/Editor/GameProjectGeneration/Classes/TemplateProjectDefs.h
+++ b/Engine/Source/Editor/GameProjectGeneration/Classes/TemplateProjectDefs.h
@@ -4,6 +4,18 @@
 
 #include "TemplateProjectDefs.generated.h"
 
+// does not require reflection exposure
+struct FTemplateConfigValue
+{
+	FString ConfigFile;
+	FString ConfigSection;
+	FString ConfigKey;
+	FString ConfigValue;
+	bool bShouldReplaceExistingValue;
+
+	GAMEPROJECTGENERATION_API FTemplateConfigValue(const FString& InFile, const FString& InSection, const FString& InKey, const FString& InValue, bool InShouldReplaceExistingValue);
+};
+
 USTRUCT()
 struct FTemplateReplacement
 {
@@ -46,7 +58,7 @@ struct FLocalizedTemplateString
 	FString Text;
 };
 
-UCLASS(config=TemplateDefs)
+UCLASS(abstract,config=TemplateDefs,MinimalAPI)
 class UTemplateProjectDefs : public UObject
 {
 	GENERATED_UCLASS_BODY()
@@ -84,6 +96,18 @@ class UTemplateProjectDefs : public UObject
 	/** Returns the display name for the current culture, or English if the current culture has no translation */
 	FText GetLocalizedDescription();
 
+	/** Does this template generate C++ source? */
+	virtual bool GeneratesCode(const FString& ProjectTemplatePath) const PURE_VIRTUAL(UTemplateProjectDefs::GeneratesCode, return false;)
+
+	/** Callback for each file rename, so class renames can be extracted*/
+	virtual bool IsClassRename(const FString& DestFilename, const FString& SrcFilename, const FString& FileExtension) const { return false; }
+
+	/** Callback for adding config values */
+	virtual void AddConfigValues(TArray<FTemplateConfigValue>& ConfigValuesToSet, const FString& TemplateName, const FString& ProjectName, bool bShouldGenerateCode) const { }
+
+	/** Callback after project generation is done, allowing for custom project generation behavior */
+	virtual bool PostGenerateProject(const FString& DestFolder, const FString& SrcFolder, const FString& NewProjectFile, const FString& TemplateFile, bool bShouldGenerateCode, FText& OutFailReason) { return true;  }
+
 private:
 	void FixString(FString& InOutStringToFix, const FString& TemplateName, const FString& ProjectName);
 };
diff --git a/Engine/Source/Editor/GameProjectGeneration/Private/DefaultTemplateProjectDefs.cpp b/Engine/Source/Editor/GameProjectGeneration/Private/DefaultTemplateProjectDefs.cpp
new file mode 100644
--- /dev/null
+++ b/Engine/Source/Editor/GameProjectGeneration/Private/DefaultTemplateProjectDefs.cpp
@@ -0,0 +1,45 @@
+// Copyright 1998-2014 Epic Games, Inc. All Rights Reserved.
+
+#include "GameProjectGenerationPrivatePCH.h"
+
+UDefaultTemplateProjectDefs::UDefaultTemplateProjectDefs(const FPostConstructInitializeProperties& PCIP)
+	: Super(PCIP)
+{
+}
+
+bool UDefaultTemplateProjectDefs::GeneratesCode(const FString& ProjectTemplatePath) const
+{
+	// Only generate code if the template has a source folder
+	const FString SourceFolder = (ProjectTemplatePath / TEXT("Source"));
+	return IFileManager::Get().DirectoryExists(*SourceFolder);
+}
+
+bool UDefaultTemplateProjectDefs::IsClassRename(const FString& DestFilename, const FString& SrcFilename, const FString& FileExtension) const
+{
+	// we shouldn't be getting this call if it's a file who's name didn't change
+	check(FPaths::GetBaseFilename(SrcFilename) != FPaths::GetBaseFilename(DestFilename));
+
+	// look for headers
+	if (FileExtension == TEXT("h"))
+	{
+		FString FileContents;
+		if (ensure(FFileHelper::LoadFileToString(FileContents, *DestFilename)))
+		{
+			// @todo uht: Checking file contents to see if this is a UObject class.  Sort of fragile here.
+			if (FileContents.Contains(TEXT(".generated.h\""), ESearchCase::IgnoreCase))
+			{
+				return true;
+			}
+		}
+	}
+
+	return false;
+}
+
+void UDefaultTemplateProjectDefs::AddConfigValues(TArray<FTemplateConfigValue>& ConfigValuesToSet, const FString& TemplateName, const FString& ProjectName, bool bShouldGenerateCode) const
+{
+	const FString ActiveGameNameRedirectsValue_LongName = FString::Printf(TEXT("(OldGameName=\"/Script/%s\",NewGameName=\"/Script/%s\")"), *TemplateName, *ProjectName);
+	const FString ActiveGameNameRedirectsValue_ShortName = FString::Printf(TEXT("(OldGameName=\"%s\",NewGameName=\"/Script/%s\")"), *TemplateName, *ProjectName);
+	new (ConfigValuesToSet) FTemplateConfigValue(TEXT("DefaultEngine.ini"), TEXT("/Script/Engine.Engine"), TEXT("+ActiveGameNameRedirects"), *ActiveGameNameRedirectsValue_LongName, /*InShouldReplaceExistingValue=*/false);
+	new (ConfigValuesToSet) FTemplateConfigValue(TEXT("DefaultEngine.ini"), TEXT("/Script/Engine.Engine"), TEXT("+ActiveGameNameRedirects"), *ActiveGameNameRedirectsValue_ShortName, /*InShouldReplaceExistingValue=*/false);
+}
diff --git a/Engine/Source/Editor/GameProjectGeneration/Private/GameProjectUtils.cpp b/Engine/Source/Editor/GameProjectGeneration/Private/GameProjectUtils.cpp
--- a/Engine/Source/Editor/GameProjectGeneration/Private/GameProjectUtils.cpp
+++ b/Engine/Source/Editor/GameProjectGeneration/Private/GameProjectUtils.cpp
@@ -5,6 +5,7 @@
 #include "UnrealEdMisc.h"
 #include "ISourceControlModule.h"
 #include "MainFrame.h"
+#include "DefaultTemplateProjectDefs.h"
 
 #include "Runtime/Analytics/Analytics/Public/Interfaces/IAnalyticsProvider.h"
 #include "EngineAnalytics.h"
@@ -807,7 +808,24 @@ UTemplateProjectDefs* GameProjectUtils::LoadTemplateDefs(const FString& ProjectD
 	const FString TemplateDefsIniFilename = ProjectDirectory / TEXT("Config") / GetTemplateDefsFilename();
 	if ( FPlatformFileManager::Get().GetPlatformFile().FileExists(*TemplateDefsIniFilename) )
 	{
-		TemplateDefs = ConstructObject<UTemplateProjectDefs>(UTemplateProjectDefs::StaticClass());
+		UClass* ClassToConstruct = UDefaultTemplateProjectDefs::StaticClass();
+
+		// see if template uses a custom project defs object
+		FString ClassName;
+		const bool bFoundValue = GConfig->GetString(*UTemplateProjectDefs::StaticClass()->GetPathName(), TEXT("TemplateProjectDefsClass"), ClassName, TemplateDefsIniFilename);
+		if (bFoundValue && ClassName.Len() > 0)
+		{
+			UClass* OverrideClass = FindObject<UClass>(ANY_PACKAGE, *ClassName, false);
+			if (nullptr != OverrideClass)
+			{
+				ClassToConstruct = OverrideClass;
+			}
+			else
+			{
+				UE_LOG(LogGameProjectGeneration, Error, TEXT("Failed to find template project defs class '%s', using default."), *ClassName);
+			}
+		}
+		TemplateDefs = ConstructObject<UTemplateProjectDefs>(ClassToConstruct);
 		TemplateDefs->LoadConfig(UTemplateProjectDefs::StaticClass(), *TemplateDefsIniFilename);
 	}
 
@@ -897,23 +915,6 @@ bool GameProjectUtils::GenerateProjectFromScratch(const FString& NewProjectFile,
 	return true;
 }
 
-struct FConfigValue
-{
-	FString ConfigFile;
-	FString ConfigSection;
-	FString ConfigKey;
-	FString ConfigValue;
-	bool bShouldReplaceExistingValue;
-
-	FConfigValue(const FString& InFile, const FString& InSection, const FString& InKey, const FString& InValue, bool InShouldReplaceExistingValue)
-		: ConfigFile(InFile)
-		, ConfigSection(InSection)
-		, ConfigKey(InKey)
-		, ConfigValue(InValue)
-		, bShouldReplaceExistingValue(InShouldReplaceExistingValue)
-	{}
-};
-
 bool GameProjectUtils::CreateProjectFromTemplate(const FString& NewProjectFile, const FString& TemplateFile, bool bShouldGenerateCode, bool bCopyStarterContent, FText& OutFailReason)
 {
 	const FString ProjectName = FPaths::GetBaseFilename(NewProjectFile);
@@ -1042,19 +1043,12 @@ bool GameProjectUtils::CreateProjectFromTemplate(const FString& NewProjectFile,
 				FilesThatNeedContentsReplaced.Add(DestFilename);
 			}
 
-			if ( FileExtension == TEXT("h")															// A header file
-				&& FPaths::GetBaseFilename(SrcFilename) != FPaths::GetBaseFilename(DestFilename))	// Whose name changed
+			// Allow project template to extract class renames from this file copy
+			if (FPaths::GetBaseFilename(SrcFilename) != FPaths::GetBaseFilename(DestFilename)
+				&& TemplateDefs->IsClassRename(DestFilename, SrcFilename, FileExtension))
 			{
-				FString FileContents;
-				if( ensure( FFileHelper::LoadFileToString( FileContents, *DestFilename ) ) )
-				{
-					// @todo uht: Checking file contents to see if this is a UObject class.  Sort of fragile here.
-					if( FileContents.Contains( TEXT( ".generated.h\"" ), ESearchCase::IgnoreCase ) )
-					{
-						// Looks like a UObject header!
-						ClassRenames.Add(FPaths::GetBaseFilename(SrcFilename), FPaths::GetBaseFilename(DestFilename));
-					}
-				}
+				// Looks like a UObject file!
+				ClassRenames.Add(FPaths::GetBaseFilename(SrcFilename), FPaths::GetBaseFilename(DestFilename));
 			}
 		}
 		else
@@ -1103,24 +1097,21 @@ bool GameProjectUtils::CreateProjectFromTemplate(const FString& NewProjectFile,
 	}
 
 	// Fixup specific ini values
-	TArray<FConfigValue> ConfigValuesToSet;
-	const FString ActiveGameNameRedirectsValue_LongName = FString::Printf(TEXT("(OldGameName=\"/Script/%s\",NewGameName=\"/Script/%s\")"), *TemplateName, *ProjectName);
-	const FString ActiveGameNameRedirectsValue_ShortName = FString::Printf(TEXT("(OldGameName=\"%s\",NewGameName=\"/Script/%s\")"), *TemplateName, *ProjectName);
-	new (ConfigValuesToSet) FConfigValue(TEXT("DefaultEngine.ini"), TEXT("/Script/Engine.Engine"), TEXT("+ActiveGameNameRedirects"), *ActiveGameNameRedirectsValue_LongName, /*InShouldReplaceExistingValue=*/false);
-	new (ConfigValuesToSet) FConfigValue(TEXT("DefaultEngine.ini"), TEXT("/Script/Engine.Engine"), TEXT("+ActiveGameNameRedirects"), *ActiveGameNameRedirectsValue_ShortName, /*InShouldReplaceExistingValue=*/false);
-	new (ConfigValuesToSet) FConfigValue(TEXT("DefaultGame.ini"), TEXT("/Script/EngineSettings.GeneralProjectSettings"), TEXT("ProjectID"), FGuid::NewGuid().ToString(), /*InShouldReplaceExistingValue=*/true);
+	TArray<FTemplateConfigValue> ConfigValuesToSet;
+	TemplateDefs->AddConfigValues(ConfigValuesToSet, TemplateName, ProjectName, bShouldGenerateCode);
+	new (ConfigValuesToSet) FTemplateConfigValue(TEXT("DefaultGame.ini"), TEXT("/Script/EngineSettings.GeneralProjectSettings"), TEXT("ProjectID"), FGuid::NewGuid().ToString(), /*InShouldReplaceExistingValue=*/true);
 
 	// Add all classname fixups
 	for ( auto RenameIt = ClassRenames.CreateConstIterator(); RenameIt; ++RenameIt )
 	{
 		const FString ClassRedirectString = FString::Printf(TEXT("(OldClassName=\"%s\",NewClassName=\"%s\")"), *RenameIt.Key(), *RenameIt.Value());
-		new (ConfigValuesToSet) FConfigValue(TEXT("DefaultEngine.ini"), TEXT("/Script/Engine.Engine"), TEXT("+ActiveClassRedirects"), *ClassRedirectString, /*InShouldReplaceExistingValue=*/false);
+		new (ConfigValuesToSet) FTemplateConfigValue(TEXT("DefaultEngine.ini"), TEXT("/Script/Engine.Engine"), TEXT("+ActiveClassRedirects"), *ClassRedirectString, /*InShouldReplaceExistingValue=*/false);
 	}
 
 	// Fix all specified config values
 	for ( auto ConfigIt = ConfigValuesToSet.CreateConstIterator(); ConfigIt; ++ConfigIt )
 	{
-		const FConfigValue& ConfigValue = *ConfigIt;
+		const FTemplateConfigValue& ConfigValue = *ConfigIt;
 		const FString IniFilename = DestFolder / TEXT("Config") / ConfigValue.ConfigFile;
 		bool bSuccessfullyProcessed = false;
 
@@ -1289,6 +1280,13 @@ bool GameProjectUtils::CreateProjectFromTemplate(const FString& NewProjectFile,
 		}
 	}
 
+	if (!TemplateDefs->PostGenerateProject(DestFolder, SrcFolder, NewProjectFile, TemplateFile, bShouldGenerateCode, OutFailReason))
+	{
+		DeleteGeneratedProjectFiles(NewProjectFile);
+		DeleteCreatedFiles(DestFolder, CreatedFiles);
+		return false;
+	}
+	
 	return true;
 }
 
diff --git a/Engine/Source/Editor/GameProjectGeneration/Private/SNewProjectWizard.cpp b/Engine/Source/Editor/GameProjectGeneration/Private/SNewProjectWizard.cpp
--- a/Engine/Source/Editor/GameProjectGeneration/Private/SNewProjectWizard.cpp
+++ b/Engine/Source/Editor/GameProjectGeneration/Private/SNewProjectWizard.cpp
@@ -976,8 +976,7 @@ void SNewProjectWizard::FindTemplateProjects()
 					}
 
 					// Only generate code if the template has a source folder
-					const FString SourceFolder = (Root / TEXT("Source"));
-					const bool bGenerateCode = IFileManager::Get().DirectoryExists(*SourceFolder);
+					const bool bGenerateCode = TemplateDefs->GeneratesCode(Root);
 
 					TSharedPtr<FSlateDynamicImageBrush> DynamicBrush;
 					const FString ThumbnailPNGFile = FPaths::GetBaseFilename(ProjectFilename, false) + TEXT(".png");
@@ -1091,17 +1090,20 @@ void SNewProjectWizard::UpdateProjectFileValidity( )
 			bLastGlobalValidityCheckSuccessful = false;
 			LastGlobalValidityErrorText = LOCTEXT("NoTemplateSelected", "No Template Selected");
 		}
-		else if ( IsCompilerRequired() )
+		else
 		{
-			if ( !FSourceCodeNavigation::IsCompilerAvailable() )
-			{
-				bLastGlobalValidityCheckSuccessful = false;
-				LastGlobalValidityErrorText = FText::Format( LOCTEXT("NoCompilerFound", "No compiler was found. In order to use a C++ template, you must first install {0}."), FSourceCodeNavigation::GetSuggestedSourceCodeIDE() );
-			}
-			else if ( !FModuleManager::Get().IsUnrealBuildToolAvailable() )
+			if (IsCompilerRequired())
 			{
-				bLastGlobalValidityCheckSuccessful = false;
-				LastGlobalValidityErrorText = LOCTEXT("UBTNotFound", "Engine source code was not found. In order to use a C++ template, you must have engine source code in Engine/Source.");
+				if ( !FSourceCodeNavigation::IsCompilerAvailable() )
+				{
+					bLastGlobalValidityCheckSuccessful = false;
+					LastGlobalValidityErrorText = FText::Format( LOCTEXT("NoCompilerFound", "No compiler was found. In order to use a C++ template, you must first install {0}."), FSourceCodeNavigation::GetSuggestedSourceCodeIDE() );
+				}
+				else if ( !FModuleManager::Get().IsUnrealBuildToolAvailable() )
+				{
+					bLastGlobalValidityCheckSuccessful = false;
+					LastGlobalValidityErrorText = LOCTEXT("UBTNotFound", "Engine source code was not found. In order to use a C++ template, you must have engine source code in Engine/Source.");
+				}
 			}
 		}
 	}
diff --git a/Engine/Source/Editor/GameProjectGeneration/Private/TemplateProjectDefs.cpp b/Engine/Source/Editor/GameProjectGeneration/Private/TemplateProjectDefs.cpp
--- a/Engine/Source/Editor/GameProjectGeneration/Private/TemplateProjectDefs.cpp
+++ b/Engine/Source/Editor/GameProjectGeneration/Private/TemplateProjectDefs.cpp
@@ -42,6 +42,18 @@ FText GetLocalizedText(const TArray<FLocalizedTemplateString>& LocalizedStrings)
 
 }
 
+FTemplateConfigValue::FTemplateConfigValue(const FString& InFile, const FString& InSection, const FString& InKey, const FString& InValue, bool InShouldReplaceExistingValue)
+	: ConfigFile(InFile)
+	, ConfigSection(InSection)
+	, ConfigKey(InKey)
+	, ConfigValue(InValue)
+	, bShouldReplaceExistingValue(InShouldReplaceExistingValue)
+{
+}
+
+
+
+
 UTemplateProjectDefs::UTemplateProjectDefs(const class FPostConstructInitializeProperties& PCIP)
 	: Super(PCIP)
 {
