Editor: Add UBlueprintFunctionLibrary to short-list for new class dialog

diff --git a/Engine/Source/Editor/GameProjectGeneration/Private/SNewClassDialog.cpp b/Engine/Source/Editor/GameProjectGeneration/Private/SNewClassDialog.cpp
--- a/Engine/Source/Editor/GameProjectGeneration/Private/SNewClassDialog.cpp
+++ b/Engine/Source/Editor/GameProjectGeneration/Private/SNewClassDialog.cpp
@@ -1223,7 +1223,8 @@ void SNewClassDialog::SetupParentClassItems()
 	FeaturedClasses.Add(GameProjectUtils::FNewClassInfo(AHUD::StaticClass()));
 	FeaturedClasses.Add(GameProjectUtils::FNewClassInfo(APlayerState::StaticClass()));
 	FeaturedClasses.Add(GameProjectUtils::FNewClassInfo(AGameState::StaticClass()));
-	
+	FeaturedClasses.Add(GameProjectUtils::FNewClassInfo(UBlueprintFunctionLibrary::StaticClass()));
+
 	// Add the extra non-UObject classes
 	FeaturedClasses.Add(GameProjectUtils::FNewClassInfo(GameProjectUtils::FNewClassInfo::EClassType::SlateWidget));
 	FeaturedClasses.Add(GameProjectUtils::FNewClassInfo(GameProjectUtils::FNewClassInfo::EClassType::SlateWidgetStyle));
