Editor: Allow 'Add Class' copyright notice to work in external projects, and change default external notice string to be 'fill this out in project settings'

diff --git a/Engine/Config/BaseGame.ini b/Engine/Config/BaseGame.ini
--- a/Engine/Config/BaseGame.ini
+++ b/Engine/Config/BaseGame.ini
@@ -24,7 +24,7 @@ bRequiresPushToTalk=true
 
 [/Script/EngineSettings.GeneralProjectSettings]
 CompanyName=
-CopyrightNotice=
+CopyrightNotice=Fill out your copyright notice in the Description page of Project Settings.
 Description=
 LicensingTerms=
 PrivacyPolicy=
diff --git a/Engine/Source/Editor/GameProjectGeneration/Private/GameProjectUtils.cpp b/Engine/Source/Editor/GameProjectGeneration/Private/GameProjectUtils.cpp
--- a/Engine/Source/Editor/GameProjectGeneration/Private/GameProjectUtils.cpp
+++ b/Engine/Source/Editor/GameProjectGeneration/Private/GameProjectUtils.cpp
@@ -2050,12 +2050,15 @@ bool GameProjectUtils::WriteOutputFile(const FString& OutputFilename, const FStr
 
 FString GameProjectUtils::MakeCopyrightLine()
 {
-	if(FEngineBuildSettings::IsInternalBuild())
+	const FString CopyrightNotice = GetDefault<UGeneralProjectSettings>()->CopyrightNotice;
+	if (!CopyrightNotice.IsEmpty())
 	{
-		return FString(TEXT("// ")) + Cast<UGeneralProjectSettings>(UGeneralProjectSettings::StaticClass()->GetDefaultObject())->CopyrightNotice;
+		return FString(TEXT("// ")) + CopyrightNotice;
+	}
+	else
+	{
+		return FString();
 	}
-
-	return "";
 }
 
 FString GameProjectUtils::MakeCommaDelimitedList(const TArray<FString>& InList, bool bPlaceQuotesAroundEveryElement)
