Reworked the new project dialog

 * Categorized templates into Blueprint/C++ templates and organized them into tabs
 * New categories can be specified in the templatedefs.ini file of a template project using a Category tag:
         Category=NewCategory
 * New category types can be registered with a proper description and icon in code through FGameProjectGenerationModule::RegisterTemplateCategory.
 * Added preview images to templates
 * Generally tidied the UI

This addresses TTP#321302 - EDITOR: Project Browser: Reduce Thumbnail size of Templates

diff --git a/Engine/Source/Editor/EditorStyle/Private/SlateEditorStyle.cpp b/Engine/Source/Editor/EditorStyle/Private/SlateEditorStyle.cpp
--- a/Engine/Source/Editor/EditorStyle/Private/SlateEditorStyle.cpp
+++ b/Engine/Source/Editor/EditorStyle/Private/SlateEditorStyle.cpp
@@ -1547,21 +1547,46 @@ void FSlateEditorStyle::FStyle::SetupGeneralStyles()
 
 		Set( "GameProjectDialog.DefaultGameThumbnail", new IMAGE_BRUSH( "GameProjectDialog/default_game_thumbnail", Icon128x128 ) );
 		Set( "GameProjectDialog.BlankProjectThumbnail", new IMAGE_BRUSH( "GameProjectDialog/blank_project_thumbnail", Icon128x128 ) );
+		Set( "GameProjectDialog.BlankProjectPreview", new IMAGE_BRUSH( "GameProjectDialog/blank_project_preview", FVector2D(512, 256) ) );
 		Set( "GameProjectDialog.BasicCodeThumbnail", new IMAGE_BRUSH( "GameProjectDialog/basic_code_thumbnail", Icon128x128 ) );
-		Set( "GameProjectDialog.CodeImage", new IMAGE_BRUSH( "GameProjectDialog/feature_code", FVector2D(70,31) ) );
+		Set( "GameProjectDialog.BasicCodePreview", new IMAGE_BRUSH( "GameProjectDialog/basic_code_preview", FVector2D(512, 256) ) );
+		Set( "GameProjectDialog.CodeImage", new IMAGE_BRUSH( "GameProjectDialog/feature_code", FVector2D(96,96) ) );
+		Set( "GameProjectDialog.BlueprintImage", new IMAGE_BRUSH( "GameProjectDialog/feature_blueprint", FVector2D(96,96) ) );
 		Set( "GameProjectDialog.CodeBorder", new BOX_BRUSH( "GameProjectDialog/feature_border", FMargin(4.0f/16.0f), FLinearColor(0.570, 0.359, 0.081, 1.f) ) );
 		Set( "GameProjectDialog.FeatureText", FTextBlockStyle(NormalText)
 			.SetFont( TTF_CORE_FONT( "Fonts/Roboto-Regular", 14 ) )
 			.SetShadowOffset( FVector2D( 0,1 ) )
 			.SetShadowColorAndOpacity( FLinearColor(0,0,0,0.9f) ) );
 		Set( "GameProjectDialog.TemplateItemTitle", FTextBlockStyle(NormalText)
-			.SetFont( TTF_CORE_FONT( "Fonts/Roboto-Regular", 14 ) )
+			.SetFont( TTF_CORE_FONT( "Fonts/Roboto-Regular", 10 ) )
 			.SetShadowOffset( FVector2D( 0,1 ) )
 			.SetShadowColorAndOpacity( FLinearColor(0,0,0,0.9f) ) );
+		
+		Set( "GameProjectDialog.Tab", FCheckBoxStyle()
+			.SetCheckBoxType( ESlateCheckBoxType::ToggleButton )
+			.SetUncheckedImage( BOX_BRUSH( "/GameProjectDialog/Tab_Inactive", 4 / 16.0f ) )
+			.SetUncheckedPressedImage( BOX_BRUSH( "/GameProjectDialog/Tab_Active", 4 / 16.0f ) )
+			.SetUncheckedHoveredImage( BOX_BRUSH( "/GameProjectDialog/Tab_Active", 4 / 16.0f ) )
+			.SetCheckedHoveredImage( BOX_BRUSH( "/GameProjectDialog/Tab_Active", 4 / 16.0f ) )
+			.SetCheckedPressedImage( BOX_BRUSH( "/GameProjectDialog/Tab_Active", 4 / 16.0f ) )
+			.SetCheckedImage( BOX_BRUSH( "/GameProjectDialog/Tab_Active", 4 / 16.0f ) )
+			);
+
+		Set( "GameProjectDialog.TabBackground", new BOX_BRUSH( "/GameProjectDialog/Tab_Background", 4 / 16.0f ) );
 
 		Set( "GameProjectDialog.FolderIconClosed", new IMAGE_BRUSH( "Icons/FolderClosed", FVector2D(18, 16) ) );
 		Set( "GameProjectDialog.FolderIconOpen", new IMAGE_BRUSH( "Icons/FolderOpen", FVector2D(18, 16) ) );
 		Set( "GameProjectDialog.ProjectFileIcon", new IMAGE_BRUSH( "Icons/doc_16x", FVector2D(18, 16) ) );
+
+		Set( "FilePath.FolderButton",
+			FButtonStyle(HoverHintOnly)
+			.SetNormal(  BOX_BRUSH( "Common/ButtonHoverHint", FMargin(4/16.0f), FLinearColor(0.5f,0.5f,0.5f) ) )
+			.SetHovered( BOX_BRUSH( "Common/ButtonHoverHint", FMargin(4/16.0f), FLinearColor(0.6f,0.6f,0.6f) ) )
+			.SetPressed( BOX_BRUSH( "Common/ButtonHoverHint", FMargin(4/16.0f), FLinearColor(0.65f,0.65f,0.65f) ) )
+			.SetNormalPadding( FMargin(0,0,0,1) )
+			.SetPressedPadding( FMargin(0,1,0,0) )
+			);
+		Set( "FilePath.GroupIndicator", new BOX_BRUSH( "GameProjectDialog/filepath_group_indicator", FMargin(4.0f/16.f) ) );
 	}
 
 	// NewClassDialog
diff --git a/Engine/Source/Editor/GameProjectGeneration/Classes/TemplateProjectDefs.h b/Engine/Source/Editor/GameProjectGeneration/Classes/TemplateProjectDefs.h
--- a/Engine/Source/Editor/GameProjectGeneration/Classes/TemplateProjectDefs.h
+++ b/Engine/Source/Editor/GameProjectGeneration/Classes/TemplateProjectDefs.h
@@ -87,6 +87,9 @@ class UTemplateProjectDefs : public UObject
 	UPROPERTY(config)
 	FString SortKey;
 
+	UPROPERTY(config)
+	FName Category;
+
 	/** Fixes up all strings in this definitions object to replace %TEMPLATENAME% with the supplied template name and %PROJECTNAME% with the supplied project name */
 	void FixupStrings(const FString& TemplateName, const FString& ProjectName);
 
diff --git a/Engine/Source/Editor/GameProjectGeneration/Private/GameProjectGenerationModule.cpp b/Engine/Source/Editor/GameProjectGeneration/Private/GameProjectGenerationModule.cpp
--- a/Engine/Source/Editor/GameProjectGeneration/Private/GameProjectGenerationModule.cpp
+++ b/Engine/Source/Editor/GameProjectGeneration/Private/GameProjectGenerationModule.cpp
@@ -3,17 +3,33 @@
 #include "GameProjectGenerationPrivatePCH.h"
 #include "ModuleManager.h"
 #include "GameProjectGenerationModule.h"
-
-
+#include "TemplateCategory.h"
+#include "SourceCodeNavigation.h"
 
 IMPLEMENT_MODULE( FGameProjectGenerationModule, GameProjectGeneration );
 DEFINE_LOG_CATEGORY(LogGameProjectGeneration);
 
 #define LOCTEXT_NAMESPACE "GameProjectGeneration"
 
+FName FTemplateCategory::BlueprintCategoryName = "Blueprint";
+FName FTemplateCategory::CodeCategoryName = "C++";
+
 void FGameProjectGenerationModule::StartupModule()
 {
-	
+	RegisterTemplateCategory(
+		FTemplateCategory::BlueprintCategoryName,
+		LOCTEXT("BlueprintCategory_Name", "Blueprint"),
+		LOCTEXT("BlueprintCategory_Description", "Blueprint templates require no programming knowledge.\nAll game mechanics can be implemented using Blueprint visual scripting.\nEach template includes a basic set of blueprints to use as a starting point for your game."),
+		FEditorStyle::GetBrush("GameProjectDialog.BlueprintImage"));
+
+	RegisterTemplateCategory(
+		FTemplateCategory::CodeCategoryName,
+		LOCTEXT("CodeCategory_Name", "C++"),
+		FText::Format(
+			LOCTEXT("CodeCategory_Description", "C++ templates offer a good example of how to work with some of the core concepts of the Engine from code.\nYou still have the option of adding your own blueprints to the project at a later date if you want.\nChoosing this template type requires you to have {0} installed."),
+			FSourceCodeNavigation::GetSuggestedSourceCodeIDE()
+		),
+		FEditorStyle::GetBrush("GameProjectDialog.CodeImage"));
 }
 
 
@@ -98,4 +114,21 @@ void FGameProjectGenerationModule::ClearSupportedTargetPlatforms()
 	GameProjectUtils::ClearSupportedTargetPlatforms();
 }
 
+bool FGameProjectGenerationModule::RegisterTemplateCategory(FName Type, FText Name, FText Description, const FSlateBrush* Thumbnail)
+{
+	if (TemplateCategories.Contains(Type))
+	{
+		return false;
+	}
+
+	FTemplateCategory Category = { Name, Description, Thumbnail, Type };
+	TemplateCategories.Add(Type, MakeShareable(new FTemplateCategory(Category)));
+	return true;
+}
+
+void FGameProjectGenerationModule::UnRegisterTemplateCategory(FName Type)
+{
+	TemplateCategories.Remove(Type);
+}
+
 #undef LOCTEXT_NAMESPACE
diff --git a/Engine/Source/Editor/GameProjectGeneration/Private/SNewProjectWizard.cpp b/Engine/Source/Editor/GameProjectGeneration/Private/SNewProjectWizard.cpp
--- a/Engine/Source/Editor/GameProjectGeneration/Private/SNewProjectWizard.cpp
+++ b/Engine/Source/Editor/GameProjectGeneration/Private/SNewProjectWizard.cpp
@@ -6,37 +6,276 @@
 #include "DesktopPlatformModule.h"
 #include "SourceCodeNavigation.h"
 #include "SScrollBorder.h"
+#include "TemplateCategory.h"
 
 #define LOCTEXT_NAMESPACE "NewProjectWizard"
 
 FName SNewProjectWizard::TemplatePageName = TEXT("Template");
 FName SNewProjectWizard::NameAndLocationPageName = TEXT("NameAndLocation");
 
+namespace
+{
+	const float ThumbnailSize = 64.f, ThumbnailPadding = 5.f;
+	const float ItemWidth = ThumbnailSize + 2*ThumbnailPadding;
+	const float ItemHeight = ItemWidth + 30;
+}
+
+/** Struct describing a single template project */
 struct FTemplateItem
 {
-	FText Name;
-	FText Description;
-	FString SortKey;
-	FString ProjectFile;
-	TSharedPtr<FSlateBrush> TemplateThumbnail;
-	bool bGenerateCode;
-
-	FTemplateItem(const FText& InName, const FText& InDescription, const TSharedPtr<FSlateBrush>& InTemplateThumbnail, bool InGenerateCode, const FString& InSortKey, const FString& InProjectFile)
-		: Name(InName)
-		, Description(InDescription)
-		, SortKey(InSortKey)
-		, ProjectFile(InProjectFile)
-		, TemplateThumbnail(InTemplateThumbnail)
-		, bGenerateCode(InGenerateCode)
+	FText		Name;
+	FText		Description;
+	bool		bGenerateCode;
+	FName		Type;
+
+	FString		SortKey;
+	FString		ProjectFile;
+
+	TSharedPtr<FSlateBrush> Thumbnail;
+	TSharedPtr<FSlateBrush> PreviewImage;
+
+	FTemplateItem(FText InName, FText InDescription, bool bInGenerateCode, FName InType, FString InSortKey, FString InProjectFile, TSharedPtr<FSlateBrush> InThumbnail, TSharedPtr<FSlateBrush> InPreviewImage)
+		: Name(InName), Description(InDescription), bGenerateCode(bInGenerateCode), Type(InType), SortKey(MoveTemp(InSortKey)), ProjectFile(MoveTemp(InProjectFile)), Thumbnail(InThumbnail), PreviewImage(InPreviewImage)
 	{}
 };
 
-struct FTemplateSortPredicate
+/**
+ * Simple widget used to display a folder path, and a name of a file:
+ * __________________________  ____________________
+ * | [REDACTED_USER_HOME]    |  | SomeFile.txt     |
+ * |-------- Folder --------|  |------ Name ------|
+ */
+class SFilepath : public SCompoundWidget
 {
-	bool operator()(const TSharedPtr<FTemplateItem>& A, const TSharedPtr<FTemplateItem>& B) const
+public:
+
+	SLATE_BEGIN_ARGS( SFilepath )
+		: _LabelBackgroundColor(FLinearColor::Black)
+		, _LabelBackgroundBrush(FEditorStyle::GetBrush("WhiteBrush"))
+	{}
+		/** Attribute specifying the text to display in the folder input */
+		SLATE_ATTRIBUTE(FText, FolderPath)
+
+		/** Attribute specifying the text to display in the name input */
+		SLATE_ATTRIBUTE(FText, Name)
+
+		/** Background label tint for the folder/name labels */
+		SLATE_ATTRIBUTE(FSlateColor, LabelBackgroundColor)
+
+		/** Background label brush for the folder/name labels */
+		SLATE_ATTRIBUTE(const FSlateBrush*, LabelBackgroundBrush)
+
+		/** Event that is triggered when the browser for folder button is clicked */
+		SLATE_EVENT(FOnClicked, OnBrowseForFolder)
+
+		/** Events for when the name field is manipulated */
+		SLATE_EVENT(FOnTextChanged, OnNameChanged)
+		SLATE_EVENT(FOnTextCommitted, OnNameCommitted)
+		
+		/** Events for when the folder field is manipulated */
+		SLATE_EVENT(FOnTextChanged, OnFolderChanged)
+		SLATE_EVENT(FOnTextCommitted, OnFolderCommitted)
+
+	SLATE_END_ARGS()
+
+	/** Constructs this widget with InArgs */
+	void Construct( const FArguments& InArgs )
 	{
-		return A->SortKey < B->SortKey;
+		ChildSlot
+		[
+			SNew(SGridPanel)
+			.FillColumn(0, 2.f)
+			.FillColumn(1, 1.f)
+
+			// Folder input
+			+ SGridPanel::Slot(0, 0)
+			[
+				SNew(SOverlay)
+
+				+ SOverlay::Slot()
+				[
+					SNew(SEditableTextBox)
+					.Text(InArgs._FolderPath)
+					// Large right hand padding to make room for the browse button
+					.Padding(FMargin(5.f, 3.f, 25.f, 3.f))
+					.OnTextChanged(InArgs._OnFolderChanged)
+					.OnTextCommitted(InArgs._OnFolderCommitted)
+				]
+					
+				+ SOverlay::Slot()
+				.HAlign(HAlign_Right)
+				[
+					SNew(SButton)
+					.ButtonStyle(FEditorStyle::Get(), "FilePath.FolderButton")
+					.ContentPadding(FMargin(4.f, 0.f))
+					.OnClicked(InArgs._OnBrowseForFolder)
+					.ToolTipText(LOCTEXT("BrowseForFolder", "Browse for a folder"))
+					.Text(LOCTEXT("...", "..."))
+				]
+			]
+
+			// Folder label
+			+ SGridPanel::Slot(0, 1)
+			[
+				SNew(SOverlay)
+
+				+ SOverlay::Slot()
+				.VAlign(VAlign_Center)
+				[
+					SNew(SBox)
+					.HeightOverride(3)
+					[
+						SNew(SBorder)
+						.BorderImage(FEditorStyle::GetBrush("FilePath.GroupIndicator"))
+						.BorderBackgroundColor(FLinearColor(1.f, 1.f, 1.f, 0.5f))
+						.Padding(FMargin(150.f, 0.f))
+					]
+				]
+					
+				+ SOverlay::Slot()
+				.HAlign(HAlign_Center)
+				.VAlign(VAlign_Center)
+				[
+					SNew(SBorder)
+					.Padding(5.f)
+					.BorderImage(InArgs._LabelBackgroundBrush)
+					.BorderBackgroundColor(InArgs._LabelBackgroundColor)
+					[
+						SNew(STextBlock)
+						.Text(LOCTEXT("Folder", "Folder"))
+					]
+				]
+			]
+
+			// Name input
+			+ SGridPanel::Slot(1, 0)
+			.Padding(FMargin(5.f, 0.f, 0.f, 0.f))
+			.VAlign(VAlign_Center)
+			[
+				SNew(SEditableTextBox)
+				.Text(InArgs._Name)
+				.Padding(FMargin(5.f, 3.f))
+				.OnTextChanged(InArgs._OnNameChanged)
+				.OnTextCommitted(InArgs._OnNameCommitted)
+			]
+
+			// Name label
+			+ SGridPanel::Slot(1, 1)
+			.Padding(FMargin(5.f, 0.f, 0.f, 0.f))
+			[
+				SNew(SOverlay)
+
+				+ SOverlay::Slot()
+				.VAlign(VAlign_Center)
+				[
+					SNew(SBox)
+					.HeightOverride(3)
+					[
+						SNew(SBorder)
+						.BorderImage(FEditorStyle::GetBrush("FilePath.GroupIndicator"))
+						.BorderBackgroundColor(FLinearColor(1.f, 1.f, 1.f, 0.5f))
+						.Padding(FMargin(75.f, 0.f))
+					]
+				]
+					
+				+ SOverlay::Slot()
+				.HAlign(HAlign_Center)
+				.VAlign(VAlign_Center)
+				[
+					SNew(SBorder)
+					.Padding(5.f)
+					.BorderImage(InArgs._LabelBackgroundBrush)
+					.BorderBackgroundColor(InArgs._LabelBackgroundColor)
+					[
+						SNew(STextBlock)
+						.Text(LOCTEXT("Name", "Name"))
+					]
+				]
+			]
+		];
+	}
+};
+
+/** Slate tile widget for template projects */
+class STemplateTile : public STableRow<TSharedPtr<FTemplateItem>>
+{
+public:
+	SLATE_BEGIN_ARGS( STemplateTile ){}
+		SLATE_ARGUMENT( TSharedPtr<FTemplateItem>, Item )
+	SLATE_END_ARGS()
+
+private:
+	TWeakPtr<FTemplateItem> Item;
+
+public:
+	/** Static build function */
+	static TSharedRef<ITableRow> BuildTile(TSharedPtr<FTemplateItem> Item, const TSharedRef<STableViewBase>& OwnerTable)
+	{
+		if (!ensure(Item.IsValid()))
+		{
+			return SNew(STableRow<TSharedPtr<FTemplateItem>>, OwnerTable);
+		}
+
+		return SNew(STemplateTile, OwnerTable).Item(Item);
+	}
+
+	/** Constructs this widget with InArgs */
+	void Construct( const FArguments& InArgs, const TSharedRef<STableViewBase>& OwnerTable )
+	{
+		check(InArgs._Item.IsValid())
+		Item = InArgs._Item;
+
+		STableRow::Construct(
+			STableRow::FArguments()
+			.Style(FEditorStyle::Get(), "GameProjectDialog.TemplateListView.TableRow")
+			.Content()
+			[
+				SNew(SVerticalBox)
+
+				// Thumbnail
+				+ SVerticalBox::Slot()
+				.AutoHeight()
+				.HAlign(HAlign_Center)
+				[
+					SNew(SBox)
+					.Padding(ThumbnailPadding)
+					.WidthOverride(ThumbnailSize)
+					.HeightOverride(ThumbnailSize)
+					[
+						SNew(SImage)
+						.Image(this, &STemplateTile::GetThumbnail)
+					]
+				]
+
+				// Name
+				+ SVerticalBox::Slot()
+				.HAlign(HAlign_Center)
+				.VAlign(VAlign_Top)
+				[
+					SNew(SRichTextBlock)
+					.WrapTextAt(ItemWidth - ThumbnailPadding*2)
+					.Justification(ETextJustify::Center)
+					//.HighlightText(this, &SNewProjectWizard::GetItemHighlightText)
+					.Text(InArgs._Item->Name)
+				]
+			],
+			OwnerTable
+		);
+	}
+
+private:
+
+	/** Get this item's thumbnail or return the default */
+	const FSlateBrush* GetThumbnail() const
+	{
+		auto ItemPtr = Item.Pin();
+		if (ItemPtr.IsValid() && ItemPtr->Thumbnail.IsValid())
+		{
+			return ItemPtr->Thumbnail.Get();
+		}
+		return FEditorStyle::GetBrush("GameProjectDialog.DefaultGameThumbnail");
 	}
+
 };
 
 BEGIN_SLATE_FUNCTION_BUILD_OPTIMIZATION
@@ -47,33 +286,24 @@ void SNewProjectWizard::Construct( const FArguments& InArgs )
 	bLastGlobalValidityCheckSuccessful = true;
 	bLastNameAndLocationValidityCheckSuccessful = true;
 	bPreventPeriodicValidityChecksUntilNextChange = false;
-	bItemsAreExpanded = false;
-	bCopyStarterContent = GEditor->GetGameAgnosticSettings().bCopyStarterContentPreference;
-	bStarterContentIsAvailable = GameProjectUtils::IsStarterContentAvailableForNewProjects();
-
-	// If the items should be initially visible, show them now
-	ExpanderRolloutCurve = FCurveSequence(0.0f, 0.1f, ECurveEaseFunction::QuadOut);
-	if(bItemsAreExpanded)
-	{
-		ExpanderRolloutCurve.JumpToEnd();
-	}
-
+	 
 	// Find all template projects
 	FindTemplateProjects();
 	SetDefaultProjectLocation();
 
-	// Width of all labels; increase this if adding a longer label to this page
-	const float LabelWidth = 46;
-	const float EditableTextHeight = 26;
-
-	SAssignNew(TemplateListView, SListView< TSharedPtr<FTemplateItem> >)
-	.ListItemsSource(&TemplateList)
+	SAssignNew(TemplateListView, STileView< TSharedPtr<FTemplateItem> >)
+	.ListItemsSource(&FilteredTemplateList)
 	.SelectionMode(ESelectionMode::Single)
 	.ClearSelectionOnClick(false)
-	.OnGenerateRow(this, &SNewProjectWizard::HandleTemplateListViewGenerateRow)
+	.OnGenerateTile_Static(&STemplateTile::BuildTile)
+	.ItemHeight(ItemHeight)
+	.ItemWidth(ItemWidth)
 	.OnMouseButtonDoubleClick(this, &SNewProjectWizard::HandleTemplateListViewDoubleClick)
 	.OnSelectionChanged(this, &SNewProjectWizard::HandleTemplateListViewSelectionChanged);
 
+	const FMargin StandardVerticalPad(0.f, 0.f, 0.f, 10.f);
+	const EVisibility StarterContentVisiblity = GameProjectUtils::IsStarterContentAvailableForNewProjects() ? EVisibility::Visible : EVisibility::Collapsed;
+
 	ChildSlot
 	[
 		SNew(SOverlay)
@@ -94,282 +324,162 @@ void SNewProjectWizard::Construct( const FArguments& InArgs )
 			+SWizard::Page()
 			.OnEnter(this, &SNewProjectWizard::OnPageVisited, TemplatePageName)
 			[
-				SNew(SVerticalBox)
-
-				// Page description
-				+SVerticalBox::Slot()
-				.AutoHeight()
-				.Padding(10.0f, 0.0f)
-				[
-					SNew(STextBlock)
-					.Text(FText::Format(LOCTEXT("ChooseTemplateDescription", "Choose a template for your project. If you choose a C++ template, you must have {0} installed to compile.\nBlueprint templates use only visual scripting, and do not require additional software."), FSourceCodeNavigation::GetSuggestedSourceCodeIDE()))
-				]
-
-				// Templates list
-				+SVerticalBox::Slot()
-				.FillHeight(1.0f)
-				.Padding(0.0f, 10.0f)
-				[
-					SNew(SScrollBorder, TemplateListView.ToSharedRef())
-					[
-						SNew(SBorder)
-						.BorderImage(FEditorStyle::GetBrush("ToolPanel.GroupBorder"))
-						[
-							TemplateListView.ToSharedRef()
-						]
-					]
-				]
-
-				+SVerticalBox::Slot()
-				.AutoHeight()
-				.Padding(0.0f, 0.0f, 0.0f, 6.0f)
-				[
-					SNew(SCheckBox)
-					.ToolTipText( LOCTEXT("CopyStarterContent_ToolTip", "Adds example content to your new project, including simple placeable meshes with basic materials and textures.  You can turn this off to start with a project that is only has the bare essentials for the selected project template."))
-					.OnCheckStateChanged(this, &SNewProjectWizard::OnCopyStarterContentChanged)
-					.IsChecked(this, &SNewProjectWizard::IsCopyStarterContentChecked)
-					.Visibility(this, &SNewProjectWizard::GetCopyStarterContentCheckboxVisibility)
-					.Content()
-					[
-						SNew(STextBlock)
-						.Text(LOCTEXT("CopyStarterContent", "Include starter content" ))
- 					]
-				]
-
-				+SVerticalBox::Slot()
-				.AutoHeight()
-				.Padding(0.0f, 0.0f, 0.0f, 8.0f)
+				SNew(SBorder)
+				.BorderImage(FEditorStyle::GetBrush("NoBorder"))
+				.Padding(10.f)
 				[
-					// Constant height, whether the label is visible or not
-					SNew(SBox)
-					.HeightOverride(20)
-					[
-						SNew(SBorder)
-						.Visibility(this, &SNewProjectWizard::GetNameAndLocationErrorLabelVisibility)
-						.BorderImage(FEditorStyle::GetBrush("GameProjectDialog.ErrorLabelBorder"))
-						.Content()
-						[
-							SNew(STextBlock)
-							.Text(this, &SNewProjectWizard::GetNameAndLocationErrorLabelText)
-							.TextStyle(FEditorStyle::Get(), "GameProjectDialog.ErrorLabelFont")
-						]
-					]
-				]
+					SNew(SHorizontalBox)
 
-				// Properties
-				+SVerticalBox::Slot()
-				.AutoHeight()
-				[
-					SNew(SBorder)
-					.BorderImage(FEditorStyle::GetBrush("DetailsView.CategoryTop"))
-					.BorderBackgroundColor(FLinearColor(0.6f, 0.6f, 0.6f, 1.0f ))
-					.Padding(FMargin(6.0f, 4.0f, 7.0f, 4.0f))
+					// Template category tabs
+					+ SHorizontalBox::Slot()
 					[
 						SNew(SVerticalBox)
 
-						// Name
 						+SVerticalBox::Slot()
+						.Padding(FMargin(5.f, 0.f))
 						.AutoHeight()
 						[
-							SNew(SHorizontalBox)
+							BuildCategoryTabs()
+						]
 
-							// Name label
-							+SHorizontalBox::Slot()
-							.AutoWidth()
-							.VAlign(VAlign_Center)
-							.Padding(0.0f, 4.0f, 6.0f, 4.0f)
+						// Templates list
+						+ SVerticalBox::Slot()
+						.FillHeight(1.0f)
+						// Manual padding offset to align to the bottom of the file path widget
+						.Padding(FMargin(0,0,0,13.f))
+						[
+							SNew(SBorder)
+							.Padding(10)
+							.BorderImage(FEditorStyle::GetBrush("GameProjectDialog.TabBackground"))
 							[
-								SNew(SBox)
-								.HAlign(HAlign_Right)
-								.WidthOverride(LabelWidth)
+								SNew(SScrollBorder, TemplateListView.ToSharedRef())
 								[
-									SNew(STextBlock)
-									.TextStyle(FEditorStyle::Get(), "GameProjectDialog.ProjectNamePathLabels")
-									.Text(LOCTEXT("ProjectNameLabel", "Name"))
+									TemplateListView.ToSharedRef()
 								]
 							]
+						]
+					]
+
+					// Selected template details
+					+ SHorizontalBox::Slot()
+					.AutoWidth()
+					// top-offset fudge to account for size of tabs
+					.Padding(FMargin(20.0f, 42.f, 0.f, 0.f))
+					[
+						// This pane is always 512px wide - the width of the template preview
+						SNew(SBox)
+						.WidthOverride(512)
+						[
+							SNew(SVerticalBox)
 
-							// Name editable text
-							+SHorizontalBox::Slot()
-							.VAlign(VAlign_Center)
-							.AutoWidth()
-							.Padding(6.0f, 4.0f, 0.0f, 4.0f)
+							// Preview image
+							+ SVerticalBox::Slot()
+							.AutoHeight()
+							.HAlign(HAlign_Center)
+							.Padding(StandardVerticalPad)
 							[
-								SNew(SBox)
-								.HeightOverride(EditableTextHeight)
-								.WidthOverride( 200.0f )	// We don't need a giant text field for the project name, so set a max width
+								SNew(SOverlay)
+							
+								+ SOverlay::Slot()
 								[
-									SNew(SEditableTextBox)
-									.Text(this, &SNewProjectWizard::GetCurrentProjectFileName)
-									.OnTextChanged(this, &SNewProjectWizard::OnCurrentProjectFileNameChanged)
+									SNew(SBorder)
+									.Padding(FMargin(0, 0, 0, 5.f))
+									.BorderImage(FEditorStyle::GetBrush("ContentBrowser.ThumbnailShadow"))
+									[
+										SNew(SImage)
+										.Image(this, &SNewProjectWizard::GetSelectedTemplatePreviewImage)
+									]
 								]
-							]
 
-							+SHorizontalBox::Slot()	// Spacer
-
-							// Expander button
-							+SHorizontalBox::Slot()
-							.AutoWidth()
-							.Padding(6.0f, 4.0f, 0.0f, 4.0f)
-							[
-								SNew(SButton)
-								.OnClicked(this, &SNewProjectWizard::OnExpanderClicked)
-								.ToolTipText(LOCTEXT("ExpanderButtonTooltip", "Click to toggle display of the full path and file creation preview"))
+								+ SOverlay::Slot()
+								.HAlign(HAlign_Right)
+								.VAlign(VAlign_Top)
+								.Padding(10)
 								[
 									SNew(SBox)
-									.HAlign(HAlign_Center)
-									.VAlign(VAlign_Center)
+									.WidthOverride(48)
+									.HeightOverride(48)
 									[
 										SNew(SImage)
-										.Image(this, &SNewProjectWizard::GetExpanderIcon)
-										.ColorAndOpacity(FLinearColor::Black)
+										.Image(this, &SNewProjectWizard::GetSelectedTemplateTypeImage)
 									]
 								]
 							]
-						]
 
-						// Expander content
-						+SVerticalBox::Slot()
-						.AutoHeight()
-						[
-							SNew(SBorder)
-							.Padding(0.0f)
-							.BorderImage(FStyleDefaults::GetNoBrush())
-							.Visibility(this, &SNewProjectWizard::GetExpandedItemsVisibility)
-							.DesiredSizeScale(this, &SNewProjectWizard::GetExpandedItemsScale)
+							// Template Name
+							+ SVerticalBox::Slot()
+							.Padding(StandardVerticalPad)
+							.AutoHeight()
 							[
-								SNew(SVerticalBox)
-
-								// Path
-								+SVerticalBox::Slot()
-								.AutoHeight()
+								SNew(STextBlock)
+								.AutoWrapText(true)
+								.TextStyle(FEditorStyle::Get(), "GameProjectDialog.FeatureText")
+								.Text(this, &SNewProjectWizard::GetSelectedTemplateProperty<FText>, &FTemplateItem::Name)
+							]
+						
+							// Template Description
+							+ SVerticalBox::Slot()
+							.Padding(StandardVerticalPad)
+							[
+								SNew(SScrollBox)
+								+ SScrollBox::Slot()
 								[
-									SNew(SHorizontalBox)
-
-									// Path label
-									+SHorizontalBox::Slot()
-									.AutoWidth()
-									.Padding(0.0f, 4.0f, 6.0f, 4.0f)
-									.VAlign(VAlign_Center)
-									[
-										SNew(SBox)
-										.HAlign(HAlign_Right)
-										.WidthOverride(LabelWidth)
-										[
-											SNew(STextBlock)
-											.TextStyle(FEditorStyle::Get(), "GameProjectDialog.ProjectNamePathLabels")
-											.Text(LOCTEXT("ProjectPathLabel", "Folder"))
-										]
-									]
-
-									// Path editable text
-									+SHorizontalBox::Slot()
-									.Padding(6.0f, 4.0f, 0.0f, 4.0f)
-									.VAlign(VAlign_Center)
-									[
-										SNew(SBox)
-										.HeightOverride(EditableTextHeight)
-										[
-											SNew(SHorizontalBox)
-
-											+SHorizontalBox::Slot()
-											.FillWidth(1.0f)
-											[
-												SNew(SEditableTextBox)
-												.Text(this, &SNewProjectWizard::GetCurrentProjectFilePath)
-												.OnTextChanged(this, &SNewProjectWizard::OnCurrentProjectFilePathChanged)
-											]
-
-											+SHorizontalBox::Slot()
-											.AutoWidth()
-											.Padding(6.0f, 1.0f, 0.0f, 0.0f)
-											[
-												SNew(SButton)
-												.VAlign(VAlign_Center)
-												.OnClicked(this, &SNewProjectWizard::HandleBrowseButtonClicked)
-												.Text(LOCTEXT( "BrowseButtonText", "Choose Folder"))
-											]
-										]
-									]
+									SNew(STextBlock)
+									.AutoWrapText(true)
+									.Text(this, &SNewProjectWizard::GetSelectedTemplateProperty<FText>, &FTemplateItem::Description)
 								]
-
-								+SVerticalBox::Slot()
-								.AutoHeight()
+							]
+						
+							// Include Starter Content?
+							+ SVerticalBox::Slot()
+							.AutoHeight()
+							.Padding(StandardVerticalPad)
+							[
+								SNew(SCheckBox)
+								.Visibility(StarterContentVisiblity)
+								.ToolTipText( LOCTEXT("CopyStarterContent_ToolTip", "Adds example content to your new project, including simple placeable meshes with basic materials and textures.\nYou can turn this off to start with a project that is only has the bare essentials for the selected project template."))
+								.OnCheckStateChanged(this, &SNewProjectWizard::OnCopyStarterContentChanged)
+								.IsChecked(this, &SNewProjectWizard::IsCopyStarterContentChecked)
+								.Content()
 								[
-									SNew(SHorizontalBox)
-
-									+SHorizontalBox::Slot()
-									.AutoWidth()
-									[
-										SNew(SVerticalBox)
-
-										// Proj file preview
-										+SVerticalBox::Slot()
-										.AutoHeight()
-										[
-											SNew(SHorizontalBox)
-
-											// Proj file preview label
-											+SHorizontalBox::Slot()
-											.AutoWidth()
-											.VAlign(VAlign_Center)
-											.Padding(0.0f, 4.0f, 6.0f, 4.0f)
-											[
-												SNew(SBox)
-												.HAlign(HAlign_Right)
-												.WidthOverride(LabelWidth)
-												[
-													SNew(STextBlock)
-													.TextStyle(FEditorStyle::Get(), "GameProjectDialog.ProjectNamePathLabels")
-													.Text(LOCTEXT("ProjFilePreviewLabel", "File"))
-												]
-											]
-
-											// Proj file preview label
-											+SHorizontalBox::Slot()
-											.VAlign(VAlign_Center)
-											.Padding(6.0f, 4.0f, 0.0f, 4.0f)
-											[
-												SNew(STextBlock)
-												.Text(this, &SNewProjectWizard::GetProjectFilenameWithPathLabelText)
-											]
-										]
-
-										// Layout explanation
-										+SVerticalBox::Slot()
-										.AutoHeight()
-										.VAlign(VAlign_Center)
-										.Padding(LabelWidth + 12.0f, 4.0f, 6.0f, 4.0f)
-										[
-											SNew(STextBlock)
-											.Text(FText::Format(LOCTEXT("ProjectNameAndLocationDetails", "A folder containing a project file (.{0}) and a few other folders will be created for you at the specified path."), FText::FromString(IProjectManager::GetProjectFileExtension())))
-										]
-									]
+									SNew(STextBlock)
+									.Text(LOCTEXT("CopyStarterContent", "Include starter content" ))
+ 								]
+							]
 
-									// Layout diagram
-									+SHorizontalBox::Slot()
-									.Padding(6.0f, 4.0f, 0.0f, 0.0f)
-									[
-										SNew(SBox)
-										[
-											ConstructDiagram()
-										]
-									]
+							// Project name error text
+							+ SVerticalBox::Slot()
+							.AutoHeight()
+							.Padding(StandardVerticalPad)
+							[
+								SNew(SBorder)
+								.BorderImage(FEditorStyle::GetBrush("GameProjectDialog.ErrorLabelBorder"))
+								.Visibility(this, &SNewProjectWizard::GetNameAndLocationErrorLabelVisibility)
+								.Padding(5.f)
+								.Content()
+								[
+									SNew(STextBlock)
+									.Text(this, &SNewProjectWizard::GetNameAndLocationErrorLabelText)
+									.TextStyle(FEditorStyle::Get(), "GameProjectDialog.ErrorLabelFont")
 								]
 							]
+
+							// File path widget
+							+ SVerticalBox::Slot()
+							.AutoHeight()
+							[
+								SNew(SFilepath)
+								.OnBrowseForFolder(this, &SNewProjectWizard::HandleBrowseButtonClicked)
+								.LabelBackgroundBrush(FEditorStyle::GetBrush("ProjectBrowser.Background"))
+								.LabelBackgroundColor(FLinearColor::White)
+								.FolderPath(this, &SNewProjectWizard::GetCurrentProjectFilePath)
+								.Name(this, &SNewProjectWizard::GetCurrentProjectFileName)
+								.OnFolderChanged(this, &SNewProjectWizard::OnCurrentProjectFilePathChanged)
+								.OnNameChanged(this, &SNewProjectWizard::OnCurrentProjectFileNameChanged)
+							]
 						]
 					]
 				]
-
-				+SVerticalBox::Slot()
-				.AutoHeight()
-				.Padding(0.0f)
-				[
-					SNew(SBorder)
-					.Padding(FMargin(0.0f, 3.0f, 0.0f, 0.0f))
-					.BorderImage(FEditorStyle::GetBrush("DetailsView.CategoryBottom"))
-					.BorderBackgroundColor(FLinearColor(0.6f, 0.6f, 0.6f, 1.0f ))
-				]
 			]
 		]
 
@@ -433,213 +543,95 @@ void SNewProjectWizard::Construct( const FArguments& InArgs )
 	// Initialize the current page name. Assuming the template page.
 	CurrentPageName = TemplatePageName;
 
-	// Select the first item
-	if (TemplateList.Num() > 0)
-	{
-		TemplateListView->SetSelection(TemplateList[0], ESelectInfo::Direct);
-	}
+	HandleCategoryChanged(ESlateCheckBoxState::Checked, ActiveCategory);
 
 	UpdateProjectFileValidity();
 }
 END_SLATE_FUNCTION_BUILD_OPTIMIZATION
 
-void SNewProjectWizard::OnCopyStarterContentChanged(ESlateCheckBoxState::Type InNewState)
-{
-	bCopyStarterContent = InNewState == ESlateCheckBoxState::Checked;
-}
-ESlateCheckBoxState::Type SNewProjectWizard::IsCopyStarterContentChecked() const
-{
-	return bCopyStarterContent ? ESlateCheckBoxState::Checked : ESlateCheckBoxState::Unchecked;
-}
 
-EVisibility SNewProjectWizard::GetCopyStarterContentCheckboxVisibility() const
+TSharedRef<SWidget> SNewProjectWizard::BuildCategoryTabs()
 {
-	return bStarterContentIsAvailable ? EVisibility::Visible : EVisibility::Collapsed;
-}
+	TSharedRef<SHorizontalBox> TabStrip = SNew(SHorizontalBox);
 
-void SNewProjectWizard::Tick( const FGeometry& AllottedGeometry, const double InCurrentTime, const float InDeltaTime )
-{
-	SCompoundWidget::Tick(AllottedGeometry, InCurrentTime, InDeltaTime);
+	TArray<FName> Categories;
+	Templates.GenerateKeyArray(Categories);
 
-	// Every few seconds, the project file path is checked for validity in case the disk contents changed and the location is now valid or invalid.
-	// After project creation, periodic checks are disabled to prevent a brief message indicating that the project you created already exists.
-	// This feature is re-enabled if the user did not restart and began editing parameters again.
-	if ( !bPreventPeriodicValidityChecksUntilNextChange && (InCurrentTime > LastValidityCheckTime + ValidityCheckFrequency) )
+	for (const FName& CategoryName : Categories)
 	{
-		UpdateProjectFileValidity();
-	}
-}
-
+		TSharedPtr<const FTemplateCategory> Category = FGameProjectGenerationModule::Get().GetCategory(CategoryName);
 
-TSharedRef<SWidget> SNewProjectWizard::ConstructDiagram()
-{
-	const FMargin FolderPadding(0.0f, 0.0f, 2.0f, 0.0f);
-
-	return SNew(SVerticalBox)
-	.ToolTipText(LOCTEXT("FileDiagramTooltip", "List of files and folders that will be created"))
-	// Parent folder
-	+SVerticalBox::Slot()
-	.AutoHeight()
-	[
-		SNew(SHorizontalBox)
+		TSharedPtr<SHorizontalBox> HorizontalBox;
 
-		+SHorizontalBox::Slot()
-		.Padding(FolderPadding)
-		.AutoWidth()
+		TabStrip->AddSlot().AutoWidth()
+		.VAlign(VAlign_Center)
+		.Padding(FMargin(0,0,2.f,0))
 		[
-			SNew(SImage)
-			.Image(FEditorStyle::GetBrush("GameProjectDialog.FolderIconOpen"))
-		]
-
-		+SHorizontalBox::Slot()
-		[
-			SNew(STextBlock)
-			.Text(this, &SNewProjectWizard::GetCurrentProjectFileParentFolder)
-		]
-	]
-
-	// Project folder
-	+SVerticalBox::Slot()
-	.Padding(10.0f, 0.0f)
-	.AutoHeight()
-	[
-		SNew(SHorizontalBox)
-
-		+SHorizontalBox::Slot()
-		.Padding(FolderPadding)
-		.AutoWidth()
-		[
-			SNew(SImage)
-			.Image(FEditorStyle::GetBrush("GameProjectDialog.FolderIconOpen"))
-		]
-
-		+SHorizontalBox::Slot()
-		[
-			SNew(STextBlock)
-			.Text(this, &SNewProjectWizard::GetCurrentProjectFileNameString)
-		]
-	]
-
-	// Config folder
-	+SVerticalBox::Slot()
-	.Padding(20.0f, 0.0f)
-	.AutoHeight()
-	[
-		SNew(SHorizontalBox)
-
-		+SHorizontalBox::Slot()
-		.Padding(FolderPadding)
-		.AutoWidth()
-		[
-			SNew(SImage)
-			.Image(FEditorStyle::GetBrush("GameProjectDialog.FolderIconClosed"))
-		]
-
-		+SHorizontalBox::Slot()
-		[
-			SNew(STextBlock)
-			.Text(FText::FromString(TEXT("Config"))) // Not localized on purpose. This folder does not change with different languages!
-		]
-	]
-
-	// Content folder
-	+SVerticalBox::Slot()
-	.Padding(20.0f, 0.0f)
-	.AutoHeight()
-	[
-		SNew(SHorizontalBox)
-
-		+SHorizontalBox::Slot()
-		.Padding(FolderPadding)
-		.AutoWidth()
-		[
-			SNew(SImage)
-			.Image(FEditorStyle::GetBrush("GameProjectDialog.FolderIconClosed"))
-		]
+			SNew(SBox)
+			// Constrain the height to 32px (for the image) plus 5px padding vertically
+			.HeightOverride(32.f + 5.f*2)
+			[
+				SNew( SCheckBox )
+				.Style( FEditorStyle::Get(), "GameProjectDialog.Tab" )
+				.OnCheckStateChanged(this, &SNewProjectWizard::HandleCategoryChanged, CategoryName)
+				.IsChecked(this, &SNewProjectWizard::GetCategoryTabCheckState, CategoryName)
+				.ToolTipText(Category.IsValid() ? Category->Description : FText())
+				.Padding(FMargin(5.f))
+				[
+					SAssignNew(HorizontalBox, SHorizontalBox)
+				]
+			]
+		];
 
-		+SHorizontalBox::Slot()
-		[
-			SNew(STextBlock)
-			.Text(FText::FromString(TEXT("Content"))) // Not localized on purpose. This folder does not change with different languages!
-		]
-	]
 
-	// Project file
-	+SVerticalBox::Slot()
-	.Padding(20.0f, 0.0f)
-	.AutoHeight()
-	[
-		SNew(SHorizontalBox)
+		if (Category.IsValid())
+		{
+			HorizontalBox->AddSlot()
+			.AutoWidth()
+			[
+				SNew(SBox)
+				.WidthOverride(32)
+				.HeightOverride(32)
+				[
+					SNew(SImage)
+					.Image(Category->Thumbnail)
+				]
+			];
+		}
 
-		+SHorizontalBox::Slot()
-		.Padding(FolderPadding)
+		HorizontalBox->AddSlot()
+		.Padding(5.f, 0.f)
+		.VAlign(VAlign_Center)
 		.AutoWidth()
-		[
-			SNew(SImage)
-			.Image(FEditorStyle::GetBrush("GameProjectDialog.ProjectFileIcon"))
-		]
-
-		+SHorizontalBox::Slot()
 		[
 			SNew(STextBlock)
-			.Text(this, &SNewProjectWizard::GetCurrentProjectFileNameStringWithExtension)
-		]
-	];
-}
-
-
-EVisibility SNewProjectWizard::GetExpandedItemsVisibility() const
-{
-	// The section is visible if its scale in Y is greater than 0
-	const float Scale = GetExpandedItemsScale().Y;
-	return (Scale > 0) ? EVisibility::Visible : EVisibility::Collapsed;
+			.TextStyle(FEditorStyle::Get(), "GameProjectDialog.FeatureText")
+			.Text(Category.IsValid() ? Category->Name : FText::FromString(CategoryName.ToString()))
+		];
+	}
+	return TabStrip;
 }
 
-
-FVector2D SNewProjectWizard::GetExpandedItemsScale() const
+void SNewProjectWizard::OnCopyStarterContentChanged(ESlateCheckBoxState::Type InNewState)
 {
-	const float Scale = ExpanderRolloutCurve.GetLerp();
-	return FVector2D(1.0f, Scale);
+	bCopyStarterContent = InNewState == ESlateCheckBoxState::Checked;
 }
-
-
-const FSlateBrush* SNewProjectWizard::GetExpanderIcon() const
+ESlateCheckBoxState::Type SNewProjectWizard::IsCopyStarterContentChecked() const
 {
-	return (bItemsAreExpanded) ? FEditorStyle::GetBrush("DetailsView.PulldownArrow.Up") : FEditorStyle::GetBrush("DetailsView.PulldownArrow.Down");
+	return bCopyStarterContent ? ESlateCheckBoxState::Checked : ESlateCheckBoxState::Unchecked;
 }
 
-
-FReply SNewProjectWizard::OnExpanderClicked()
+void SNewProjectWizard::Tick( const FGeometry& AllottedGeometry, const double InCurrentTime, const float InDeltaTime )
 {
-	bItemsAreExpanded = !bItemsAreExpanded;
-
-	if(bItemsAreExpanded)
-	{
-		ExpanderRolloutCurve = FCurveSequence(0.0f, 0.1f, ECurveEaseFunction::CubicOut);
-		ExpanderRolloutCurve.Play();
-	}
-	else
-	{
-		ExpanderRolloutCurve = FCurveSequence(0.0f, 0.1f, ECurveEaseFunction::CubicIn);
-		ExpanderRolloutCurve.PlayReverse();
-	}
-
-	return FReply::Handled();
-}
-
+	SCompoundWidget::Tick(AllottedGeometry, InCurrentTime, InDeltaTime);
 
-const FSlateBrush* SNewProjectWizard::GetTemplateItemImage(TWeakPtr<FTemplateItem> TemplateItem) const
-{
-	if ( TemplateItem.IsValid() )
+	// Every few seconds, the project file path is checked for validity in case the disk contents changed and the location is now valid or invalid.
+	// After project creation, periodic checks are disabled to prevent a brief message indicating that the project you created already exists.
+	// This feature is re-enabled if the user did not restart and began editing parameters again.
+	if ( !bPreventPeriodicValidityChecksUntilNextChange && (InCurrentTime > LastValidityCheckTime + ValidityCheckFrequency) )
 	{
-		TSharedPtr<FTemplateItem> Item = TemplateItem.Pin();
-		if ( Item->TemplateThumbnail.IsValid() )
-		{
-			return Item->TemplateThumbnail.Get();
-		}
+		UpdateProjectFileValidity();
 	}
-	
-	return FEditorStyle::GetBrush("GameProjectDialog.DefaultGameThumbnail");
 }
 
 
@@ -660,28 +652,30 @@ TSharedPtr<FTemplateItem> SNewProjectWizard::GetSelectedTemplateItem() const
 	return NULL;
 }
 
+const FSlateBrush* SNewProjectWizard::GetSelectedTemplatePreviewImage() const
+{
+	auto PreviewImage = GetSelectedTemplateProperty(&FTemplateItem::PreviewImage);
+	return PreviewImage.IsValid() ? PreviewImage.Get() : nullptr;
+}
 
-FText SNewProjectWizard::GetSelectedTemplateName() const
+const FSlateBrush* SNewProjectWizard::GetSelectedTemplateTypeImage() const
 {
 	TSharedPtr<FTemplateItem> SelectedItem = GetSelectedTemplateItem();
-	if ( SelectedItem.IsValid() )
+	if (SelectedItem.IsValid())
 	{
-		return SelectedItem->Name;
+		auto Category = FGameProjectGenerationModule::Get().GetCategory(SelectedItem->Type);
+		if (Category.IsValid())
+		{
+			return Category->Thumbnail;
+		}
 	}
-
-	return FText::GetEmpty();
+	return nullptr;
 }
 
 
 FText SNewProjectWizard::GetCurrentProjectFileName() const
 {
-	return FText::FromString( GetCurrentProjectFileNameString() );
-}
-
-
-FString SNewProjectWizard::GetCurrentProjectFileNameString() const
-{
-	return CurrentProjectFileName;
+	return FText::FromString( CurrentProjectFileName );
 }
 
 
@@ -879,7 +873,7 @@ FReply SNewProjectWizard::OnCloseGlobalErrorLabelClicked()
 
 EVisibility SNewProjectWizard::GetNameAndLocationErrorLabelVisibility() const
 {
-	return GetNameAndLocationErrorLabelText().IsEmpty() ? EVisibility::Hidden : EVisibility::Visible;
+	return GetNameAndLocationErrorLabelText().IsEmpty() ? EVisibility::Collapsed : EVisibility::Visible;
 }
 
 
@@ -895,23 +889,28 @@ FText SNewProjectWizard::GetNameAndLocationErrorLabelText() const
 
 void SNewProjectWizard::FindTemplateProjects()
 {
+	// Default to showing the blueprint category
+	ActiveCategory = FTemplateCategory::BlueprintCategoryName;
+
 	// Add some default non-data driven templates
-	TemplateList.Add( MakeShareable(new FTemplateItem(
+	Templates.FindOrAdd(FTemplateCategory::BlueprintCategoryName).Add(MakeShareable(new FTemplateItem(
 		LOCTEXT("BlankProjectName", "Blank"),
 		LOCTEXT("BlankProjectDescription", "A clean empty project with no code."),
-		MakeShareable( new FSlateBrush( *FEditorStyle::GetBrush("GameProjectDialog.BlankProjectThumbnail") ) ),
-		false,				// bGenerateCode
+		false, FTemplateCategory::BlueprintCategoryName,
 		TEXT("_1"),			// SortKey
-		TEXT("")			// No filename, this is a generation template
+		TEXT(""),			// No filename, this is a generation template
+		MakeShareable( new FSlateBrush( *FEditorStyle::GetBrush("GameProjectDialog.BlankProjectThumbnail") ) ),
+		MakeShareable( new FSlateBrush( *FEditorStyle::GetBrush("GameProjectDialog.BlankProjectPreview") ) )
 		)) );
 
-	TemplateList.Add( MakeShareable(new FTemplateItem(
+	Templates.FindOrAdd(FTemplateCategory::CodeCategoryName).Add(MakeShareable(new FTemplateItem(
 		LOCTEXT("BasicCodeProjectName", "Basic Code"),
 		LOCTEXT("BasicCodeProjectDescription", "An empty project with some basic game framework code classes created."),
-		MakeShareable( new FSlateBrush( *FEditorStyle::GetBrush("GameProjectDialog.BasicCodeThumbnail") ) ),
-		true,				// bGenerateCode
+		true, FTemplateCategory::CodeCategoryName,
 		TEXT("_2"),			// SortKey
-		TEXT("")			// No filename, this is a generation template
+		TEXT(""),			// No filename, this is a generation template
+		MakeShareable( new FSlateBrush( *FEditorStyle::GetBrush("GameProjectDialog.BasicCodeThumbnail") ) ),
+		MakeShareable( new FSlateBrush( *FEditorStyle::GetBrush("GameProjectDialog.BlankProjectPreview") ) )
 		)) );
 
 	// Now discover and all data driven templates
@@ -978,12 +977,20 @@ void SNewProjectWizard::FindTemplateProjects()
 					// Only generate code if the template has a source folder
 					const bool bGenerateCode = TemplateDefs->GeneratesCode(Root);
 
-					TSharedPtr<FSlateDynamicImageBrush> DynamicBrush;
+					TSharedPtr<FSlateDynamicImageBrush> ThumbnailBrush;
 					const FString ThumbnailPNGFile = FPaths::GetBaseFilename(ProjectFilename, false) + TEXT(".png");
 					if ( FPlatformFileManager::Get().GetPlatformFile().FileExists(*ThumbnailPNGFile) )
 					{
 						const FName BrushName = FName(*ThumbnailPNGFile);
-						DynamicBrush = MakeShareable( new FSlateDynamicImageBrush(BrushName , FVector2D(128,128) ) );
+						ThumbnailBrush = MakeShareable( new FSlateDynamicImageBrush(BrushName , FVector2D(128,128) ) );
+					}
+
+					TSharedPtr<FSlateDynamicImageBrush> PreviewBrush;
+					const FString PreviewPNGFile = FPaths::GetBaseFilename(ProjectFilename, false) + TEXT("_Preview.png");
+					if ( FPlatformFileManager::Get().GetPlatformFile().FileExists(*PreviewPNGFile) )
+					{
+						const FName BrushName = FName(*PreviewPNGFile);
+						PreviewBrush = MakeShareable( new FSlateDynamicImageBrush(BrushName , FVector2D(512,256) ) );
 					}
 
 					// Get the sort key
@@ -997,7 +1004,23 @@ void SNewProjectWizard::FindTemplateProjects()
 						SortKey = TEXT("_0");
 					}
 
-					TemplateList.Add(MakeShareable(new FTemplateItem(TemplateName, TemplateDescription, DynamicBrush, bGenerateCode, SortKey, ProjectFilename)));
+					// Assign the template to the correct category. If the template has no explicit category assigned, we assign it to either C++ or blueprint
+					FName Category = TemplateDefs->Category;
+					if (Category.IsNone())
+					{
+						Category = bGenerateCode ? FTemplateCategory::CodeCategoryName : FTemplateCategory::BlueprintCategoryName;
+					}
+
+					Templates.FindOrAdd(Category).Add(MakeShareable(new FTemplateItem(
+						TemplateName,
+						TemplateDescription,
+						bGenerateCode,
+						Category,
+						MoveTemp(SortKey),
+						MoveTemp(ProjectFilename),
+						ThumbnailBrush,
+						PreviewBrush
+					)));
 				}
 			}
 			else
@@ -1007,9 +1030,6 @@ void SNewProjectWizard::FindTemplateProjects()
 			}
 		}
 	}
-
-	// Sort the template folders
-	TemplateList.Sort(FTemplateSortPredicate());
 }
 
 
@@ -1317,83 +1337,31 @@ void SNewProjectWizard::HandleCreateProjectWizardFinished( )
 	CreateAndOpenProject();
 }
 
+ESlateCheckBoxState::Type SNewProjectWizard::GetCategoryTabCheckState(FName Category) const
+{
+	return Category == ActiveCategory ? ESlateCheckBoxState::Checked : ESlateCheckBoxState::Unchecked;
+}
 
-TSharedRef<ITableRow> SNewProjectWizard::HandleTemplateListViewGenerateRow( TSharedPtr<FTemplateItem> TemplateItem, const TSharedRef<STableViewBase>& OwnerTable )
+void SNewProjectWizard::HandleCategoryChanged(ESlateCheckBoxState::Type CheckState, FName Category)
 {
-	if (!ensure(TemplateItem.IsValid()))
+	if (CheckState != ESlateCheckBoxState::Checked)
 	{
-		return SNew(STableRow<TSharedPtr<FTemplateItem>>, OwnerTable);
+		return;
 	}
 
-	const FString CodeTemplateToolTip = FText::Format( LOCTEXT("CodeTemplateToolTip", "This template contains C++ source code files. To use this template you must have {0} installed."), FSourceCodeNavigation::GetSuggestedSourceCodeIDE() ).ToString();
-	const int32 ThumbnailBorderPadding = 5;
-	const int32 ThumbnailSize = 128;
+	ActiveCategory = Category;
+	FilteredTemplateList = Templates.FindRef(Category);
 
-	return SNew( STableRow<TSharedPtr<FTemplateItem>>, OwnerTable )
-		.Style(FEditorStyle::Get(), "GameProjectDialog.TemplateListView.TableRow")
-		[
-			SNew(SBox)
-				.HeightOverride(ThumbnailSize+ThumbnailBorderPadding+5)
-				[
-					SNew(SHorizontalBox)
-
-					// Thumbnail
-					+ SHorizontalBox::Slot()
-						.AutoWidth()
-						[
-							SNew(SBox)
-								.Padding(ThumbnailBorderPadding)
-								.WidthOverride(ThumbnailSize+ThumbnailBorderPadding * 2)
-								.HeightOverride(ThumbnailSize+ThumbnailBorderPadding * 2)
-								[
-									SNew(SImage)
-										.Image(this, &SNewProjectWizard::GetTemplateItemImage, TWeakPtr<FTemplateItem>(TemplateItem))
-								]
-						]
-
-					// Name and description
-					+ SHorizontalBox::Slot()
-						.FillWidth(1.0f)
-						.Padding(10.0f)
-						.VAlign(VAlign_Center)
-						[
-							SNew(SVerticalBox)
-
-							+ SVerticalBox::Slot()
-								.Padding(0.0f, 0.0f, 0.0f, 8.0f)
-								.AutoHeight()
-								[
-									SNew(SHorizontalBox)
-
-									+ SHorizontalBox::Slot()
-										.VAlign(VAlign_Center)
-										[
-											SNew(STextBlock)
-												.TextStyle(FEditorStyle::Get(), "GameProjectDialog.TemplateItemTitle")
-												.Text(TemplateItem->Name.ToString())
-										]
-						
-									+ SHorizontalBox::Slot()
-										.AutoWidth()
-										.VAlign(VAlign_Center)
-										[
-											SNew(SImage)
-												.Visibility(TemplateItem->bGenerateCode ? EVisibility::Visible : EVisibility::Collapsed)
-												.ToolTipText(CodeTemplateToolTip)
-												.Image(FEditorStyle::GetBrush("GameProjectDialog.CodeImage"))
-										]
-								]
+	// Sort the template folders
+	FilteredTemplateList.Sort([](const TSharedPtr<FTemplateItem>& A, const TSharedPtr<FTemplateItem>& B){
+		return A->SortKey < B->SortKey;
+	});
 
-								+ SVerticalBox::Slot()
-									.FillHeight(1.0f)
-									[
-										SNew(STextBlock)
-											.WrapTextAt(596)
-											.Text(TemplateItem->Description.ToString())
-									]
-						]
-				]
-		];
+	if (FilteredTemplateList.Num() > 0)
+	{
+		TemplateListView->SetSelection(FilteredTemplateList[0]);
+	}
+	TemplateListView->RequestListRefresh();
 }
 
 
diff --git a/Engine/Source/Editor/GameProjectGeneration/Private/SNewProjectWizard.h b/Engine/Source/Editor/GameProjectGeneration/Private/SNewProjectWizard.h
--- a/Engine/Source/Editor/GameProjectGeneration/Private/SNewProjectWizard.h
+++ b/Engine/Source/Editor/GameProjectGeneration/Private/SNewProjectWizard.h
@@ -4,6 +4,7 @@
 
 
 struct FTemplateItem;
+struct FTemplateCategory;
 
 
 /**
@@ -25,21 +26,28 @@ public:
 
 private:
 
-	/** Construct the folder diagram */
-	TSharedRef<SWidget> ConstructDiagram();
-
-	/** Gets the image to display for the specified template */
-	const FSlateBrush* GetTemplateItemImage(TWeakPtr<FTemplateItem> TemplateItem) const;
+	/** Build the set of template category tabs */
+	TSharedRef<SWidget> BuildCategoryTabs();
 
 	/** Accessor for the currently selected template item */
 	TSharedPtr<FTemplateItem> GetSelectedTemplateItem() const;
 
+	/** Helper function to allow direct lookup of the selected item's properties on a delegate */
+	template<typename T, typename TRet = T>
+	TRet GetSelectedTemplateProperty(T FTemplateItem::*Prop) const
+	{
+		TSharedPtr<FTemplateItem> SelectedItem = GetSelectedTemplateItem();
+		if ( SelectedItem.IsValid() )
+		{
+			return (*SelectedItem).*Prop;
+		}
+
+		return TRet();
+	}
+
 	/** Accessor for the project name text */
 	FText GetCurrentProjectFileName() const;
 
-	/** Accessor for the project name string */
-	FString GetCurrentProjectFileNameString() const;
-
 	/** Accessor for the project filename string with the extension */
 	FString GetCurrentProjectFileNameStringWithExtension() const;
 
@@ -61,6 +69,10 @@ private:
 	/** Accessor for the label to show the currently selected template */
 	FText GetSelectedTemplateName() const;
 
+	/** Get the images for the selected template preview and category */
+	const FSlateBrush* GetSelectedTemplatePreviewImage() const;
+	const FSlateBrush* GetSelectedTemplateTypeImage() const;
+
 	/** Gets the assembled project filename with path */
 	FString GetProjectFilenameWithPath() const;
 
@@ -149,27 +161,17 @@ private:
 	/** Handler for when the Create button is clicked */
 	void HandleCreateProjectWizardFinished( );
 
-	/** Creates a row in the template list */
-	TSharedRef<ITableRow> HandleTemplateListViewGenerateRow( TSharedPtr<FTemplateItem> TemplateItem, const TSharedRef<STableViewBase>& OwnerTable );
-
 	/** Called when a user double-clicks a template item in the template project list. */
 	void HandleTemplateListViewDoubleClick( TSharedPtr<FTemplateItem> TemplateItem );
 
 	/** Handler for when the selection changes in the template list */
 	void HandleTemplateListViewSelectionChanged( TSharedPtr<FTemplateItem> TemplateItem, ESelectInfo::Type SelectInfo );
 
-private:
-	/* @return Whether the expanded items are currently visible, based on the current expander state */
-	EVisibility GetExpandedItemsVisibility() const;
+	/** Handle choosing a different category tab */
+	void HandleCategoryChanged(ESlateCheckBoxState::Type Checked, FName Category);
 
-	/* @return The scale to apply to the expanded items, based on the rollout curve value */
-	FVector2D GetExpandedItemsScale() const;
-
-	/* @return The brush to use for the expander icon, based on the current expander state */
-	const FSlateBrush* GetExpanderIcon() const;
-
-	/* Handle the expander button being clicked (toggles the expansion state) */
-	FReply OnExpanderClicked();
+	/** Get the check state for the specified categories tab */
+	ESlateCheckBoxState::Type GetCategoryTabCheckState(FName Category) const;
 
 private:
 
@@ -202,24 +204,23 @@ private:
 	/** True if the last NameAndLocation validity check returned that the project path is valid for creation */
 	bool bLastNameAndLocationValidityCheckSuccessful;
 
-	/** True if the items are currently expanded */
-	bool bItemsAreExpanded;
-
-	/** Curved used to simulate a rollout of the section */
-	FCurveSequence ExpanderRolloutCurve;
-
 	/** The name of the page that is currently in view */
 	FName CurrentPageName;
 
 	/** True if user has selected to copy starter content. */
 	bool bCopyStarterContent;
 
-	/** True if starter content is available to be copied into new projects. */
-	bool bStarterContentIsAvailable;
+	/** Name of the currently selected category */
+	FName ActiveCategory;
+
+	/** A map of category name to array of templates available for that category */
+	TMap<FName, TArray<TSharedPtr<FTemplateItem>> > Templates;
+
+	/** The filtered array of templates we are currently showing */
+	TArray<TSharedPtr<FTemplateItem> > FilteredTemplateList;
 
-	/** Template items */
-	TArray<TSharedPtr<FTemplateItem> > TemplateList;
-	TSharedPtr<SListView<TSharedPtr<FTemplateItem> > > TemplateListView;
+	/** The slate widget representing the list of templates */
+	TSharedPtr<STileView<TSharedPtr<FTemplateItem> > > TemplateListView;
 
 	/** Names for pages */
 	static FName TemplatePageName;
diff --git a/Engine/Source/Editor/GameProjectGeneration/Private/TemplateCategory.h b/Engine/Source/Editor/GameProjectGeneration/Private/TemplateCategory.h
new file mode 100644
--- /dev/null
+++ b/Engine/Source/Editor/GameProjectGeneration/Private/TemplateCategory.h
@@ -0,0 +1,20 @@
+// Copyright 1998-2014 Epic Games, Inc. All Rights Reserved.
+
+#pragma once
+
+struct FTemplateCategory
+{
+	static FName BlueprintCategoryName, CodeCategoryName;
+
+	/** Localised name of this category */
+	FText Name;
+	
+	/** A description of the templates contained within this category */
+	FText Description;
+
+	/** A thumbnail to help identify this category */
+	const FSlateBrush* Thumbnail;
+
+	/** A unique name for this category */
+	FName Type;
+};
\ No newline at end of file
diff --git a/Engine/Source/Editor/GameProjectGeneration/Public/GameProjectGenerationModule.h b/Engine/Source/Editor/GameProjectGeneration/Public/GameProjectGenerationModule.h
--- a/Engine/Source/Editor/GameProjectGeneration/Public/GameProjectGenerationModule.h
+++ b/Engine/Source/Editor/GameProjectGeneration/Public/GameProjectGenerationModule.h
@@ -5,6 +5,8 @@
 
 #include "ModuleInterface.h"
 
+struct FTemplateCategory;
+
 /**
  * Game Project Generation module
  */
@@ -12,6 +14,7 @@ class FGameProjectGenerationModule : public IModuleInterface
 {
 
 public:
+	typedef TMap<FName, TSharedPtr<FTemplateCategory>> FTemplateCategoryMap;
 
 	/**
 	 * Called right after the plugin DLL has been loaded and the plugin object has been created
@@ -80,7 +83,18 @@ public:
 	/** Clear the list of supported target platforms */
 	virtual void ClearSupportedTargetPlatforms();
 
+public:
+
+	/** (Un)register a new type of template category to be shown on the new project page */
+	virtual bool RegisterTemplateCategory(FName Type, FText Name, FText Description, const FSlateBrush* Thumbnail);
+	virtual void UnRegisterTemplateCategory(FName Type);
+
+	// Non DLL-exposed access to template categories
+	TSharedPtr<const FTemplateCategory> GetCategory(FName Type) const { return TemplateCategories.FindRef(Type); }
+
 private:
 	FAddCodeToProjectDialogOpenedEvent AddCodeToProjectDialogOpenedEvent;
 
+	/** Map of template categories from type to ptr */
+	FTemplateCategoryMap TemplateCategories;
 };
diff --git a/Engine/Source/Editor/MainFrame/Private/Frame/MainFrameActions.cpp b/Engine/Source/Editor/MainFrame/Private/Frame/MainFrameActions.cpp
--- a/Engine/Source/Editor/MainFrame/Private/Frame/MainFrameActions.cpp
+++ b/Engine/Source/Editor/MainFrame/Private/Frame/MainFrameActions.cpp
@@ -381,8 +381,8 @@ void FMainFrameActionCallbacks::NewProject( bool bAllowProjectOpening, bool bAll
 	TSharedRef<SWindow> NewProjectWindow =
 		SNew(SWindow)
 		.Title(LOCTEXT( "SelectProjectWindowHeader", "Select Project"))
-		.ClientSize( FVector2D(1280, 720) )
-		.SizingRule( ESizingRule::FixedSize )
+		.ClientSize( FMainFrameModule::GetProjectBrowserWindowSize() )
+		.SizingRule( ESizingRule::UserSized )
 		.SupportsMinimize(false) .SupportsMaximize(false);
 
 	NewProjectWindow->SetContent( FGameProjectGenerationModule::Get().CreateGameProjectDialog(bAllowProjectOpening, bAllowProjectCreate) );
diff --git a/Engine/Source/Editor/MainFrame/Private/MainFrameModule.cpp b/Engine/Source/Editor/MainFrame/Private/MainFrameModule.cpp
--- a/Engine/Source/Editor/MainFrame/Private/MainFrameModule.cpp
+++ b/Engine/Source/Editor/MainFrame/Private/MainFrameModule.cpp
@@ -58,7 +58,7 @@ void FMainFrameModule::CreateDefaultMainFrame( const bool bStartImmersivePIE )
 			// Do not maximize the window initially. Keep a small dialog feel.
 			DefaultWindowLocation.InitiallyMaximized = false;
 
-			DefaultWindowLocation.WindowSize = FVector2D(920, 700);
+			DefaultWindowLocation.WindowSize = GetProjectBrowserWindowSize();
 
 			// Do not let the user adjust the fixed size
 			bIsUserSizable = true;
diff --git a/Engine/Source/Editor/MainFrame/Private/MainFrameModule.h b/Engine/Source/Editor/MainFrame/Private/MainFrameModule.h
--- a/Engine/Source/Editor/MainFrame/Private/MainFrameModule.h
+++ b/Engine/Source/Editor/MainFrame/Private/MainFrameModule.h
@@ -121,6 +121,11 @@ protected:
 	 */
 	bool ShouldShowProjectDialogAtStartup( ) const;
 
+public:
+
+	/** Get the size of the project browser window */
+	static FVector2D GetProjectBrowserWindowSize() { return FVector2D(1100, 700); }
+
 private:
 
 	// Handles the level editor module starting to recompile.
