Clang requires full template specification.

diff --git a/Engine/Source/Editor/GameProjectGeneration/Private/SNewProjectWizard.cpp b/Engine/Source/Editor/GameProjectGeneration/Private/SNewProjectWizard.cpp
--- a/Engine/Source/Editor/GameProjectGeneration/Private/SNewProjectWizard.cpp
+++ b/Engine/Source/Editor/GameProjectGeneration/Private/SNewProjectWizard.cpp
@@ -414,7 +414,7 @@ void SNewProjectWizard::Construct( const FArguments& InArgs )
 								SNew(STextBlock)
 								.AutoWrapText(true)
 								.TextStyle(FEditorStyle::Get(), "GameProjectDialog.FeatureText")
-								.Text(this, &SNewProjectWizard::GetSelectedTemplateProperty<FText>, &FTemplateItem::Name)
+								.Text(this, &SNewProjectWizard::GetSelectedTemplateProperty<FText, FText>, &FTemplateItem::Name)
 							]
 						
 							// Template Description
@@ -426,7 +426,7 @@ void SNewProjectWizard::Construct( const FArguments& InArgs )
 								[
 									SNew(STextBlock)
 									.AutoWrapText(true)
-									.Text(this, &SNewProjectWizard::GetSelectedTemplateProperty<FText>, &FTemplateItem::Description)
+									.Text(this, &SNewProjectWizard::GetSelectedTemplateProperty<FText, FText>, &FTemplateItem::Description)
 								]
 							]
 						
