Removed default template parameter from function to fix VS2012

diff --git a/Engine/Source/Editor/GameProjectGeneration/Private/SNewProjectWizard.cpp b/Engine/Source/Editor/GameProjectGeneration/Private/SNewProjectWizard.cpp
--- a/Engine/Source/Editor/GameProjectGeneration/Private/SNewProjectWizard.cpp
+++ b/Engine/Source/Editor/GameProjectGeneration/Private/SNewProjectWizard.cpp
@@ -414,7 +414,7 @@ void SNewProjectWizard::Construct( const FArguments& InArgs )
 								SNew(STextBlock)
 								.AutoWrapText(true)
 								.TextStyle(FEditorStyle::Get(), "GameProjectDialog.FeatureText")
-								.Text(this, &SNewProjectWizard::GetSelectedTemplateProperty<FText, FText>, &FTemplateItem::Name)
+								.Text(this, &SNewProjectWizard::GetSelectedTemplateProperty<FText>, &FTemplateItem::Name)
 							]
 						
 							// Template Description
@@ -426,7 +426,7 @@ void SNewProjectWizard::Construct( const FArguments& InArgs )
 								[
 									SNew(STextBlock)
 									.AutoWrapText(true)
-									.Text(this, &SNewProjectWizard::GetSelectedTemplateProperty<FText, FText>, &FTemplateItem::Description)
+									.Text(this, &SNewProjectWizard::GetSelectedTemplateProperty<FText>, &FTemplateItem::Description)
 								]
 							]
 						
diff --git a/Engine/Source/Editor/GameProjectGeneration/Private/SNewProjectWizard.h b/Engine/Source/Editor/GameProjectGeneration/Private/SNewProjectWizard.h
--- a/Engine/Source/Editor/GameProjectGeneration/Private/SNewProjectWizard.h
+++ b/Engine/Source/Editor/GameProjectGeneration/Private/SNewProjectWizard.h
@@ -33,8 +33,9 @@ private:
 	TSharedPtr<FTemplateItem> GetSelectedTemplateItem() const;
 
 	/** Helper function to allow direct lookup of the selected item's properties on a delegate */
-	template<typename T, typename TRet = T>
-	TRet GetSelectedTemplateProperty(T FTemplateItem::*Prop) const
+	/** TRet should be defaulted but VS2012 doesn't allow default template arguments on non-class templates */
+	template<typename T>
+	T GetSelectedTemplateProperty(T FTemplateItem::*Prop) const
 	{
 		TSharedPtr<FTemplateItem> SelectedItem = GetSelectedTemplateItem();
 		if ( SelectedItem.IsValid() )
@@ -42,7 +43,7 @@ private:
 			return (*SelectedItem).*Prop;
 		}
 
-		return TRet();
+		return T();
 	}
 
 	/** Accessor for the project name text */
