Use DesktopPlatform functions to generate project files, and include the log in error messages if it fails.

diff --git a/Engine/Source/Editor/GameProjectGeneration/Private/GameProjectUtils.cpp b/Engine/Source/Editor/GameProjectGeneration/Private/GameProjectUtils.cpp
--- a/Engine/Source/Editor/GameProjectGeneration/Private/GameProjectUtils.cpp
+++ b/Engine/Source/Editor/GameProjectGeneration/Private/GameProjectUtils.cpp
@@ -17,6 +17,7 @@
 #include "ClassIconFinder.h"
 
 #include "UProjectInfo.h"
+#include "DesktopPlatformModule.h"
 
 #define LOCTEXT_NAMESPACE "GameProjectUtils"
 
@@ -591,23 +592,6 @@ bool GameProjectUtils::CreateProject(const FString& NewProjectFile, const FStrin
 	return bProjectCreationSuccessful;
 }
 
-bool GameProjectUtils::BuildGameBinaries(const FString& ProjectFilename, FText& OutFailReason)
-{
-	const bool bAllowNewSlowTask = true;
-	FScopedSlowTask SlowTaskMessage( LOCTEXT( "BuildingProjectStatus", "Building project..." ), bAllowNewSlowTask );
-
-	// Compile the *editor* for the project
-	if ( FModuleManager::Get().CompileGameProjectEditor(ProjectFilename, *GLog) )
-	{
-		return true;
-	}
-
-	FFormatNamedArguments Args;
-	Args.Add( TEXT("ProjectFilename"), FText::FromString( ProjectFilename ) );
-	OutFailReason = FText::Format( LOCTEXT("FailedToCompileNewProject", "Failed to compile {ProjectFileName}."), Args );
-	return false;
-}
-
 void GameProjectUtils::CheckForOutOfDateGameProjectFile()
 {
 	if ( FPaths::IsProjectFilePathSet() )
@@ -1686,15 +1670,21 @@ bool GameProjectUtils::GenerateGameFrameworkSourceCode(const FString& NewProject
 
 bool GameProjectUtils::GenerateCodeProjectFiles(const FString& ProjectFilename, FText& OutFailReason)
 {
-	if ( FModuleManager::Get().GenerateCodeProjectFiles(ProjectFilename, *GLog) )
+	FStringOutputDevice OutputLog;
+	OutputLog.SetAutoEmitLineTerminator(true);
+	GLog->AddOutputDevice(&OutputLog);
+	bool bHaveProjectFiles = FDesktopPlatformModule::Get()->GenerateProjectFiles(FPaths::RootDir(), ProjectFilename, GWarn);
+	GLog->RemoveOutputDevice(&OutputLog);
+
+	if(!bHaveProjectFiles)
 	{
-		return true;
+		FFormatNamedArguments Args;
+		Args.Add( TEXT("LogOutput"), FText::FromString(OutputLog) );
+		OutFailReason = FText::Format(LOCTEXT("CouldNotGenerateProjectFiles", "Failed to generate project files. Log output:\n{LogOutput}"), Args);
+		return false;
 	}
 
-	FFormatNamedArguments Args;
-	Args.Add( TEXT("ProjectFilename"), FText::FromString( ProjectFilename ) );
-	OutFailReason = FText::Format( LOCTEXT("FailedToGenerateCodeProjectFiles", "Failed to generate code project files for \"{ProjectFilename}\"."), Args );
-	return false;
+	return true;
 }
 
 bool GameProjectUtils::IsStarterContentAvailableForNewProjects()
@@ -2675,7 +2665,7 @@ bool GameProjectUtils::AddCodeToProject_Internal(const FString& NewClassName, co
 	}
 
 	// Generate project files if we happen to be using a project file.
-	if ( !FModuleManager::Get().GenerateCodeProjectFiles( FPaths::GetProjectFilePath(), *GLog ) )
+	if ( !FDesktopPlatformModule::Get()->GenerateProjectFiles(FPaths::RootDir(), FPaths::GetProjectFilePath(), GWarn) )
 	{
 		OutFailReason = LOCTEXT("FailedToGenerateProjectFiles", "Failed to generate project files.");
 		return false;
diff --git a/Engine/Source/Editor/GameProjectGeneration/Private/GameProjectUtils.h b/Engine/Source/Editor/GameProjectGeneration/Private/GameProjectUtils.h
--- a/Engine/Source/Editor/GameProjectGeneration/Private/GameProjectUtils.h
+++ b/Engine/Source/Editor/GameProjectGeneration/Private/GameProjectUtils.h
@@ -135,9 +135,6 @@ public:
 	/** Creates the specified project file and all required folders. If TemplateFile is non-empty, it will be used as the template for creation. On failure, OutFailReason will be populated. */
 	static bool CreateProject(const FString& NewProjectFile, const FString& TemplateFile, bool bShouldGenerateCode, bool bCopyStarterContent, FText& OutFailReason);
 
-	/** Builds the binaries for a new project. */
-	static bool BuildGameBinaries(const FString& ProjectFilename, FText& OutFailReason);
-
 	/** Prompts the user to update his project file, if necessary. */
 	static void CheckForOutOfDateGameProjectFile();
 
diff --git a/Engine/Source/Editor/GameProjectGeneration/Private/SNewProjectWizard.cpp b/Engine/Source/Editor/GameProjectGeneration/Private/SNewProjectWizard.cpp
--- a/Engine/Source/Editor/GameProjectGeneration/Private/SNewProjectWizard.cpp
+++ b/Engine/Source/Editor/GameProjectGeneration/Private/SNewProjectWizard.cpp
@@ -1319,8 +1319,16 @@ void SNewProjectWizard::CloseWindowIfAppropriate( bool ForceClose )
 
 void SNewProjectWizard::DisplayError( const FText& ErrorText )
 {
-	UE_LOG(LogGameProjectGeneration, Log, TEXT("%s"), *ErrorText.ToString());
-	PersistentGlobalErrorLabelText = ErrorText;
+	FString ErrorString = ErrorText.ToString();
+	UE_LOG(LogGameProjectGeneration, Log, TEXT("%s"), *ErrorString);
+	if(ErrorString.Contains("\n"))
+	{
+		FMessageDialog::Open(EAppMsgType::Ok, ErrorText);
+	}
+	else
+	{
+		PersistentGlobalErrorLabelText = ErrorText;
+	}
 }
 
 
diff --git a/Engine/Source/Editor/MainFrame/Private/Frame/MainFrameActions.cpp b/Engine/Source/Editor/MainFrame/Private/Frame/MainFrameActions.cpp
--- a/Engine/Source/Editor/MainFrame/Private/Frame/MainFrameActions.cpp
+++ b/Engine/Source/Editor/MainFrame/Private/Frame/MainFrameActions.cpp
@@ -692,7 +692,7 @@ void FMainFrameActionCallbacks::RefreshCodeProject()
 	FText FailReason;
 	if(!FGameProjectGenerationModule::Get().UpdateCodeProject(FailReason))
 	{
-		FSlateNotificationManager::Get().AddNotification(FNotificationInfo(FailReason));
+		FMessageDialog::Open(EAppMsgType::Ok, FailReason);
 	}
 }
 
