Analytics - Log size of source directory when the number of source files are logged

diff --git a/Engine/Source/Editor/GameProjectGeneration/Private/GameProjectGenerationModule.cpp b/Engine/Source/Editor/GameProjectGeneration/Private/GameProjectGenerationModule.cpp
--- a/Engine/Source/Editor/GameProjectGeneration/Private/GameProjectGenerationModule.cpp
+++ b/Engine/Source/Editor/GameProjectGeneration/Private/GameProjectGenerationModule.cpp
@@ -90,6 +90,10 @@ int32 FGameProjectGenerationModule::GetProjectCodeFileCount()
 	return GameProjectUtils::GetProjectCodeFileCount();
 }
 
+void FGameProjectGenerationModule::GetProjectSourceDirectoryInfo(int32& OutNumFiles, int64& OutDirectorySize)
+{
+	GameProjectUtils::GetProjectSourceDirectoryInfo(OutNumFiles, OutDirectorySize);
+}
 
 bool FGameProjectGenerationModule::UpdateCodeResourceFiles(TArray<FString>& OutCreatedFiles, FText& OutFailReason)
 {
diff --git a/Engine/Source/Editor/GameProjectGeneration/Private/GameProjectUtils.cpp b/Engine/Source/Editor/GameProjectGeneration/Private/GameProjectUtils.cpp
--- a/Engine/Source/Editor/GameProjectGeneration/Private/GameProjectUtils.cpp
+++ b/Engine/Source/Editor/GameProjectGeneration/Private/GameProjectUtils.cpp
@@ -2567,15 +2567,32 @@ FString GameProjectUtils::GetDefaultProjectTemplateFilename()
 	return TEXT("");
 }
 
+void GameProjectUtils::GetProjectCodeFilenames(TArray<FString>& OutProjectCodeFilenames)
+{
+	IFileManager::Get().FindFilesRecursive(OutProjectCodeFilenames, *FPaths::GameSourceDir(), TEXT("*.h"), true, false, false);
+	IFileManager::Get().FindFilesRecursive(OutProjectCodeFilenames, *FPaths::GameSourceDir(), TEXT("*.cpp"), true, false, false);
+}
+
 int32 GameProjectUtils::GetProjectCodeFileCount()
 {
 	TArray<FString> Filenames;
-	IFileManager::Get().FindFilesRecursive(Filenames, *FPaths::GameSourceDir(), TEXT("*.h"), true, false, false);
-	IFileManager::Get().FindFilesRecursive(Filenames, *FPaths::GameSourceDir(), TEXT("*.cpp"), true, false, false);
-
+	GetProjectCodeFilenames(Filenames);
 	return Filenames.Num();
 }
 
+void GameProjectUtils::GetProjectSourceDirectoryInfo(int32& OutNumCodeFiles, int64& OutDirectorySize)
+{
+	TArray<FString> Filenames;
+	GetProjectCodeFilenames(Filenames);
+	OutNumCodeFiles = Filenames.Num();
+
+	OutDirectorySize = 0;
+	for (const auto& filename : Filenames)
+	{
+		OutDirectorySize += IFileManager::Get().FileSize(*filename);
+	}
+}
+
 bool GameProjectUtils::ProjectHasCodeFiles()
 {
 	return GameProjectUtils::GetProjectCodeFileCount() > 0;
diff --git a/Engine/Source/Editor/GameProjectGeneration/Private/GameProjectUtils.h b/Engine/Source/Editor/GameProjectGeneration/Private/GameProjectUtils.h
--- a/Engine/Source/Editor/GameProjectGeneration/Private/GameProjectUtils.h
+++ b/Engine/Source/Editor/GameProjectGeneration/Private/GameProjectUtils.h
@@ -159,9 +159,16 @@ public:
 	/** Loads a template project definitions object from the TemplateDefs.ini file in the specified project */
 	static UTemplateProjectDefs* LoadTemplateDefs(const FString& ProjectDirectory);
 
-	/** Returns number of code files in the currently loaded project */
+	/** @return The number of code files in the currently loaded project */
 	static int32 GetProjectCodeFileCount();
 
+	/** 
+	* Retrieves file and size info about the project's source directory
+	* @param OutNumFiles Contains the number of files within the source directory
+	* @param OutDirectorySize Contains the combined size of all files in the directory
+	*/
+	static void GetProjectSourceDirectoryInfo(int32& OutNumFiles, int64& OutDirectorySize);
+
 	/** Returns the uproject template filename for the default project template. */
 	static FString GetDefaultProjectTemplateFilename();
 
@@ -330,6 +337,9 @@ private:
 	/** Handler for when the user confirms a project update */
 	static void OnUpdateProjectConfirm();
 
+	/** @param OutProjectCodeFilenames Contains the filenames of the project source code files */
+	static void GetProjectCodeFilenames(TArray<FString>& OutProjectCodeFilenames);
+
 	/**
 	 * Updates the projects, and optionally the modules names
 	 *
diff --git a/Engine/Source/Editor/GameProjectGeneration/Public/GameProjectGenerationModule.h b/Engine/Source/Editor/GameProjectGeneration/Public/GameProjectGenerationModule.h
--- a/Engine/Source/Editor/GameProjectGeneration/Public/GameProjectGenerationModule.h
+++ b/Engine/Source/Editor/GameProjectGeneration/Public/GameProjectGenerationModule.h
@@ -65,6 +65,9 @@ public:
 	/** Gets the current projects source file count */
 	virtual int32 GetProjectCodeFileCount();
 
+	/** Gets file and size info about the source directory */
+	virtual void GetProjectSourceDirectoryInfo(int32& OutNumFiles, int64& OutDirectorySize);
+
 	/** Update code resource files */
 	virtual bool UpdateCodeResourceFiles(TArray<FString>& OutCreatedFiles, FText& OutFailReason);
 
diff --git a/Engine/Source/Editor/UnrealEd/Private/UnrealEdMisc.cpp b/Engine/Source/Editor/UnrealEd/Private/UnrealEdMisc.cpp
--- a/Engine/Source/Editor/UnrealEd/Private/UnrealEdMisc.cpp
+++ b/Engine/Source/Editor/UnrealEd/Private/UnrealEdMisc.cpp
@@ -360,9 +360,15 @@ void FUnrealEdMisc::InitEngineAnalytics()
 			ProjectAttributes.Add( FAnalyticsEventAttribute( FString( "Id" ), *GetDefault<UGeneralProjectSettings>()->ProjectID.ToString() ));
 
 			FGameProjectGenerationModule& GameProjectModule = FModuleManager::LoadModuleChecked<FGameProjectGenerationModule>(TEXT("GameProjectGeneration"));
-			int32 SourceFileCount = GameProjectModule.Get().GetProjectCodeFileCount();
+			
+			int32 SourceFileCount = 0;
+			int64 SourceFileDirectorySize = 0;
+			GameProjectModule.Get().GetProjectSourceDirectoryInfo(SourceFileCount, SourceFileDirectorySize);
+
 			ProjectAttributes.Add( FAnalyticsEventAttribute( FString( "SourceFileCount" ), SourceFileCount ));
+			ProjectAttributes.Add(FAnalyticsEventAttribute(FString("SourceFileDirectorySize"), SourceFileDirectorySize));
 			ProjectAttributes.Add( FAnalyticsEventAttribute( FString( "ModuleCount" ), FModuleManager::Get().GetModuleCount() ));
+
 			// UObject class count
 			int32 UObjectClasses = 0;
 			int32 UBlueprintClasses = 0;
