#ue4 - Creating an empty project no longer generates an empty PlayerController class file.

diff --git a/Engine/Source/Editor/GameProjectGeneration/Private/GameProjectUtils.cpp b/Engine/Source/Editor/GameProjectGeneration/Private/GameProjectUtils.cpp
--- a/Engine/Source/Editor/GameProjectGeneration/Private/GameProjectUtils.cpp
+++ b/Engine/Source/Editor/GameProjectGeneration/Private/GameProjectUtils.cpp
@@ -1591,22 +1591,6 @@ bool GameProjectUtils::GenerateGameFrameworkSourceCode(const FString& NewProject
 	NewModuleInfo.ModuleType = EHostType::Runtime;
 	NewModuleInfo.ModuleSourcePath = FPaths::ConvertRelativePathToFull(GameModulePath / ""); // Ensure trailing /
 
-	// MyGamePlayerController.h
-	{
-		const UClass* BaseClass = APlayerController::StaticClass();
-		const FString NewClassName = NewProjectName + BaseClass->GetName();
-		const FString NewHeaderFilename = GameModulePath / NewClassName + TEXT(".h");
-		FString UnusedSyncLocation;
-		if ( GenerateClassHeaderFile(NewHeaderFilename, NewClassName, FNewClassInfo(BaseClass), TArray<FString>(), TEXT(""), TEXT(""), UnusedSyncLocation, NewModuleInfo, OutFailReason) )
-		{
-			OutCreatedFiles.Add(NewHeaderFilename);
-		}
-		else
-		{
-			return false;
-		}
-	}
-
 	// MyGameGameMode.h
 	{
 		const UClass* BaseClass = AGameMode::StaticClass();
@@ -1623,23 +1607,6 @@ bool GameProjectUtils::GenerateGameFrameworkSourceCode(const FString& NewProject
 		}
 	}
 
-	// MyGamePlayerController.cpp
-	FString PrefixedPlayerControllerClassName;
-	{
-		const UClass* BaseClass = APlayerController::StaticClass();
-		const FString NewClassName = NewProjectName + BaseClass->GetName();
-		const FString NewCPPFilename = GameModulePath / NewClassName + TEXT(".cpp");
-		PrefixedPlayerControllerClassName = FString(BaseClass->GetPrefixCPP()) + NewClassName;
-		if ( GenerateClassCPPFile(NewCPPFilename, NewClassName, FNewClassInfo(BaseClass), TArray<FString>(), TArray<FString>(), TEXT(""), NewModuleInfo, OutFailReason) )
-		{
-			OutCreatedFiles.Add(NewCPPFilename);
-		}
-		else
-		{
-			return false;
-		}
-	}
-
 	// MyGameGameMode.cpp
 	{
 		const UClass* BaseClass = AGameMode::StaticClass();
@@ -1647,13 +1614,7 @@ bool GameProjectUtils::GenerateGameFrameworkSourceCode(const FString& NewProject
 		const FString NewCPPFilename = GameModulePath / NewClassName + TEXT(".cpp");
 		
 		TArray<FString> PropertyOverrides;
-		PropertyOverrides.Add( FString::Printf( TEXT("PlayerControllerClass = %s::StaticClass();"), *PrefixedPlayerControllerClassName ) );
-
-		// PropertyOverrides references PlayerController class so we need to include its header to properly compile under non-unity
-		const UClass* PlayerControllerBaseClass = APlayerController::StaticClass();
-		const FString PlayerControllerClassName = NewProjectName + PlayerControllerBaseClass->GetName() + TEXT(".h");
 		TArray<FString> AdditionalIncludes;
-		AdditionalIncludes.Add(PlayerControllerClassName);
 
 		if ( GenerateClassCPPFile(NewCPPFilename, NewClassName, FNewClassInfo(BaseClass), AdditionalIncludes, PropertyOverrides, TEXT(""), NewModuleInfo, OutFailReason) )
 		{
