Seperated out the New Project icon and overlay images into seperate brushes so they could easily be replaced

diff --git a/Engine/Source/Editor/EditorStyle/Private/SlateEditorStyle.cpp b/Engine/Source/Editor/EditorStyle/Private/SlateEditorStyle.cpp
--- a/Engine/Source/Editor/EditorStyle/Private/SlateEditorStyle.cpp
+++ b/Engine/Source/Editor/EditorStyle/Private/SlateEditorStyle.cpp
@@ -1553,7 +1553,9 @@ void FSlateEditorStyle::FStyle::SetupGeneralStyles()
 		Set( "GameProjectDialog.BlankProjectPreview", new IMAGE_BRUSH( "GameProjectDialog/blank_project_preview", FVector2D(512, 256) ) );
 		Set( "GameProjectDialog.BasicCodeThumbnail", new IMAGE_BRUSH( "GameProjectDialog/basic_code_thumbnail", Icon128x128 ) );
 		Set( "GameProjectDialog.BasicCodePreview", new IMAGE_BRUSH( "GameProjectDialog/basic_code_preview", FVector2D(512, 256) ) );
+		Set( "GameProjectDialog.CodeIcon", new IMAGE_BRUSH( "GameProjectDialog/feature_code_32x", FVector2D(32,32) ) );
 		Set( "GameProjectDialog.CodeImage", new IMAGE_BRUSH( "GameProjectDialog/feature_code", FVector2D(96,96) ) );
+		Set( "GameProjectDialog.BlueprintIcon", new IMAGE_BRUSH( "GameProjectDialog/feature_blueprint_32x", FVector2D(32,32) ) );
 		Set( "GameProjectDialog.BlueprintImage", new IMAGE_BRUSH( "GameProjectDialog/feature_blueprint", FVector2D(96,96) ) );
 		Set( "GameProjectDialog.CodeBorder", new BOX_BRUSH( "GameProjectDialog/feature_border", FMargin(4.0f/16.0f), FLinearColor(0.570, 0.359, 0.081, 1.f) ) );
 		Set( "GameProjectDialog.FeatureText", FTextBlockStyle(NormalText)
diff --git a/Engine/Source/Editor/GameProjectGeneration/Private/GameProjectGenerationModule.cpp b/Engine/Source/Editor/GameProjectGeneration/Private/GameProjectGenerationModule.cpp
--- a/Engine/Source/Editor/GameProjectGeneration/Private/GameProjectGenerationModule.cpp
+++ b/Engine/Source/Editor/GameProjectGeneration/Private/GameProjectGenerationModule.cpp
@@ -20,6 +20,7 @@ void FGameProjectGenerationModule::StartupModule()
 		FTemplateCategory::BlueprintCategoryName,
 		LOCTEXT("BlueprintCategory_Name", "Blueprint"),
 		LOCTEXT("BlueprintCategory_Description", "Blueprint templates require no programming knowledge.\nAll game mechanics can be implemented using Blueprint visual scripting.\nEach template includes a basic set of blueprints to use as a starting point for your game."),
+		FEditorStyle::GetBrush("GameProjectDialog.BlueprintIcon"),
 		FEditorStyle::GetBrush("GameProjectDialog.BlueprintImage"));
 
 	RegisterTemplateCategory(
@@ -29,6 +30,7 @@ void FGameProjectGenerationModule::StartupModule()
 			LOCTEXT("CodeCategory_Description", "C++ templates offer a good example of how to work with some of the core concepts of the Engine from code.\nYou still have the option of adding your own blueprints to the project at a later date if you want.\nChoosing this template type requires you to have {0} installed."),
 			FSourceCodeNavigation::GetSuggestedSourceCodeIDE()
 		),
+		FEditorStyle::GetBrush("GameProjectDialog.CodeIcon"),
 		FEditorStyle::GetBrush("GameProjectDialog.CodeImage"));
 }
 
@@ -118,14 +120,14 @@ void FGameProjectGenerationModule::ClearSupportedTargetPlatforms()
 	GameProjectUtils::ClearSupportedTargetPlatforms();
 }
 
-bool FGameProjectGenerationModule::RegisterTemplateCategory(FName Type, FText Name, FText Description, const FSlateBrush* Thumbnail)
+bool FGameProjectGenerationModule::RegisterTemplateCategory(FName Type, FText Name, FText Description, const FSlateBrush* Icon, const FSlateBrush* Image)
 {
 	if (TemplateCategories.Contains(Type))
 	{
 		return false;
 	}
 
-	FTemplateCategory Category = { Name, Description, Thumbnail, Type };
+	FTemplateCategory Category = { Name, Description, Icon, Image, Type };
 	TemplateCategories.Add(Type, MakeShareable(new FTemplateCategory(Category)));
 	return true;
 }
diff --git a/Engine/Source/Editor/GameProjectGeneration/Private/SNewProjectWizard.cpp b/Engine/Source/Editor/GameProjectGeneration/Private/SNewProjectWizard.cpp
--- a/Engine/Source/Editor/GameProjectGeneration/Private/SNewProjectWizard.cpp
+++ b/Engine/Source/Editor/GameProjectGeneration/Private/SNewProjectWizard.cpp
@@ -598,7 +598,7 @@ TSharedRef<SWidget> SNewProjectWizard::BuildCategoryTabs()
 				.HeightOverride(32)
 				[
 					SNew(SImage)
-					.Image(Category->Thumbnail)
+					.Image(Category->Icon)
 				]
 			];
 		}
@@ -670,7 +670,7 @@ const FSlateBrush* SNewProjectWizard::GetSelectedTemplateTypeImage() const
 		auto Category = FGameProjectGenerationModule::Get().GetCategory(SelectedItem->Type);
 		if (Category.IsValid())
 		{
-			return Category->Thumbnail;
+			return Category->Image;
 		}
 	}
 	return nullptr;
diff --git a/Engine/Source/Editor/GameProjectGeneration/Private/TemplateCategory.h b/Engine/Source/Editor/GameProjectGeneration/Private/TemplateCategory.h
--- a/Engine/Source/Editor/GameProjectGeneration/Private/TemplateCategory.h
+++ b/Engine/Source/Editor/GameProjectGeneration/Private/TemplateCategory.h
@@ -12,8 +12,11 @@ struct FTemplateCategory
 	/** A description of the templates contained within this category */
 	FText Description;
 
-	/** A thumbnail to help identify this category */
-	const FSlateBrush* Thumbnail;
+	/** A thumbnail to help identify this category (on the tab) */
+	const FSlateBrush* Icon;
+
+	/** A thumbnail to help identify this category (on the screenshot)*/
+	const FSlateBrush* Image;
 
 	/** A unique name for this category */
 	FName Type;
diff --git a/Engine/Source/Editor/GameProjectGeneration/Public/GameProjectGenerationModule.h b/Engine/Source/Editor/GameProjectGeneration/Public/GameProjectGenerationModule.h
--- a/Engine/Source/Editor/GameProjectGeneration/Public/GameProjectGenerationModule.h
+++ b/Engine/Source/Editor/GameProjectGeneration/Public/GameProjectGenerationModule.h
@@ -89,7 +89,7 @@ public:
 public:
 
 	/** (Un)register a new type of template category to be shown on the new project page */
-	virtual bool RegisterTemplateCategory(FName Type, FText Name, FText Description, const FSlateBrush* Thumbnail);
+	virtual bool RegisterTemplateCategory(FName Type, FText Name, FText Description, const FSlateBrush* Icon, const FSlateBrush* Image);
 	virtual void UnRegisterTemplateCategory(FName Type);
 
 	// Non DLL-exposed access to template categories
