Remove virtual function to return .uproject extension from IProjectManager.

diff --git a/Engine/Source/Editor/GameProjectGeneration/Private/GameProjectUtils.cpp b/Engine/Source/Editor/GameProjectGeneration/Private/GameProjectUtils.cpp
--- a/Engine/Source/Editor/GameProjectGeneration/Private/GameProjectUtils.cpp
+++ b/Engine/Source/Editor/GameProjectGeneration/Private/GameProjectUtils.cpp
@@ -322,10 +322,10 @@ bool GameProjectUtils::IsValidProjectFileForCreation(const FString& ProjectFile,
 		return false;
 	}
 
-	if ( FPaths::GetExtension(ProjectFile) != IProjectManager::GetProjectFileExtension() )
+	if ( FPaths::GetExtension(ProjectFile) != FProjectDescriptor::GetExtension() )
 	{
 		FFormatNamedArguments Args;
-		Args.Add( TEXT("ProjectFileExtension"), FText::FromString( IProjectManager::GetProjectFileExtension() ) );
+		Args.Add( TEXT("ProjectFileExtension"), FText::FromString( FProjectDescriptor::GetExtension() ) );
 		OutFailReason = FText::Format( LOCTEXT( "InvalidProjectFileExtension", "File extension is not {ProjectFileExtension}" ), Args );
 		return false;
 	}
@@ -369,7 +369,7 @@ bool GameProjectUtils::IsValidProjectFileForCreation(const FString& ProjectFile,
 	if ( AnyProjectFilesExistInFolder(FPaths::GetPath(ProjectFile)) )
 	{
 		FFormatNamedArguments Args;
-		Args.Add( TEXT("ProjectFileExtension"), FText::FromString( IProjectManager::GetProjectFileExtension() ) );
+		Args.Add( TEXT("ProjectFileExtension"), FText::FromString( FProjectDescriptor::GetExtension() ) );
 		Args.Add( TEXT("ProjectFilePath"), FText::FromString( FPaths::GetPath(ProjectFile) ) );
 		OutFailReason = FText::Format( LOCTEXT( "AProjectFileAlreadyExistsAtLoction", "Another .{ProjectFileExtension} file already exists in {ProjectFilePath}" ), Args );
 		return false;
@@ -408,10 +408,10 @@ bool GameProjectUtils::OpenProject(const FString& ProjectFile, FText& OutFailRea
 		return false;
 	}
 
-	if ( FPaths::GetExtension(ProjectFile) != IProjectManager::GetProjectFileExtension() )
+	if ( FPaths::GetExtension(ProjectFile) != FProjectDescriptor::GetExtension() )
 	{
 		FFormatNamedArguments Args;
-		Args.Add( TEXT("ProjectFileExtension"), FText::FromString( IProjectManager::GetProjectFileExtension() ) );
+		Args.Add( TEXT("ProjectFileExtension"), FText::FromString( FProjectDescriptor::GetExtension() ) );
 		OutFailReason = FText::Format( LOCTEXT( "InvalidProjectFileExtension", "File extension is not {ProjectFileExtension}" ), Args );
 		return false;
 	}
@@ -1335,7 +1335,7 @@ bool GameProjectUtils::ProjectFileExists(const FString& ProjectFile)
 bool GameProjectUtils::AnyProjectFilesExistInFolder(const FString& Path)
 {
 	TArray<FString> ExistingFiles;
-	const FString Wildcard = FString::Printf(TEXT("%s/*.%s"), *Path, *IProjectManager::GetProjectFileExtension());
+	const FString Wildcard = FString::Printf(TEXT("%s/*.%s"), *Path, *FProjectDescriptor::GetExtension());
 	IFileManager::Get().FindFiles(ExistingFiles, *Wildcard, /*Files=*/true, /*Directories=*/false);
 
 	return ExistingFiles.Num() > 0;
diff --git a/Engine/Source/Editor/GameProjectGeneration/Private/SNewProjectWizard.cpp b/Engine/Source/Editor/GameProjectGeneration/Private/SNewProjectWizard.cpp
--- a/Engine/Source/Editor/GameProjectGeneration/Private/SNewProjectWizard.cpp
+++ b/Engine/Source/Editor/GameProjectGeneration/Private/SNewProjectWizard.cpp
@@ -685,7 +685,7 @@ FText SNewProjectWizard::GetCurrentProjectFileName() const
 
 FString SNewProjectWizard::GetCurrentProjectFileNameStringWithExtension() const
 {
-	return CurrentProjectFileName + TEXT(".") + IProjectManager::GetProjectFileExtension();
+	return CurrentProjectFileName + TEXT(".") + FProjectDescriptor::GetExtension();
 }
 
 
@@ -740,7 +740,7 @@ FString SNewProjectWizard::GetProjectFilenameWithPath() const
 	{
 		const FString ProjectName = CurrentProjectFileName;
 		const FString ProjectPath = IFileManager::Get().ConvertToAbsolutePathForExternalAppForWrite(*CurrentProjectFilePath);
-		const FString Filename = ProjectName + TEXT(".") + IProjectManager::GetProjectFileExtension();
+		const FString Filename = ProjectName + TEXT(".") + FProjectDescriptor::GetExtension();
 		FString ProjectFilename = FPaths::Combine( *ProjectPath, *ProjectName, *Filename );
 		FPaths::MakePlatformFilename(ProjectFilename);
 		return ProjectFilename;
@@ -955,7 +955,7 @@ void SNewProjectWizard::FindTemplateProjects()
 	// Add a template item for every discovered project
 	for ( auto TemplateFolderIt = AllTemplateFolders.CreateConstIterator(); TemplateFolderIt; ++TemplateFolderIt )
 	{
-		const FString SearchString = (*TemplateFolderIt) / TEXT("*.") + IProjectManager::GetProjectFileExtension();
+		const FString SearchString = (*TemplateFolderIt) / TEXT("*.") + FProjectDescriptor::GetExtension();
 		TArray<FString> FoundProjectFiles;
 		IFileManager::Get().FindFiles(FoundProjectFiles, *SearchString, /*Files=*/true, /*Directories=*/false);
 		if ( FoundProjectFiles.Num() > 0 )
@@ -1070,7 +1070,7 @@ void SNewProjectWizard::SetDefaultProjectLocation( )
 
 	// Check to make sure the project file doesn't already exist
 	FText FailReason;
-	if ( !GameProjectUtils::IsValidProjectFileForCreation(DefaultProjectFilePath / ProjectName / ProjectName + TEXT(".") + IProjectManager::GetProjectFileExtension(), FailReason) )
+	if ( !GameProjectUtils::IsValidProjectFileForCreation(DefaultProjectFilePath / ProjectName / ProjectName + TEXT(".") + FProjectDescriptor::GetExtension(), FailReason) )
 	{
 		// If it exists, find an appropriate numerical suffix
 		const int MaxSuffix = 1000;
@@ -1078,7 +1078,7 @@ void SNewProjectWizard::SetDefaultProjectLocation( )
 		for ( Suffix = 2; Suffix < MaxSuffix; ++Suffix )
 		{
 			ProjectName = GenericProjectName + FString::Printf(TEXT("%d"), Suffix);
-			if ( GameProjectUtils::IsValidProjectFileForCreation(DefaultProjectFilePath / ProjectName / ProjectName + TEXT(".") + IProjectManager::GetProjectFileExtension(), FailReason) )
+			if ( GameProjectUtils::IsValidProjectFileForCreation(DefaultProjectFilePath / ProjectName / ProjectName + TEXT(".") + FProjectDescriptor::GetExtension(), FailReason) )
 			{
 				// Found a name that is not taken. Break out.
 				break;
diff --git a/Engine/Source/Editor/GameProjectGeneration/Private/SProjectBrowser.cpp b/Engine/Source/Editor/GameProjectGeneration/Private/SProjectBrowser.cpp
--- a/Engine/Source/Editor/GameProjectGeneration/Private/SProjectBrowser.cpp
+++ b/Engine/Source/Editor/GameProjectGeneration/Private/SProjectBrowser.cpp
@@ -1077,7 +1077,7 @@ void SProjectBrowser::HandleProjectItemDoubleClick( TSharedPtr<FProjectItem> Tem
 FReply SProjectBrowser::OnBrowseToProjectClicked()
 {
 	const FString ProjectFileDescription = LOCTEXT( "FileTypeDescription", "Unreal Project File" ).ToString();
-	const FString ProjectFileExtension = FString::Printf(TEXT("*.%s"), *IProjectManager::GetProjectFileExtension());
+	const FString ProjectFileExtension = FString::Printf(TEXT("*.%s"), *FProjectDescriptor::GetExtension());
 	const FString FileTypes = FString::Printf( TEXT("%s (%s)|%s"), *ProjectFileDescription, *ProjectFileExtension, *ProjectFileExtension );
 
 	// Find the first valid project file to select by default
diff --git a/Engine/Source/Editor/GameProjectGeneration/Private/Tests/GameProjectAutomationTests.cpp b/Engine/Source/Editor/GameProjectGeneration/Private/Tests/GameProjectAutomationTests.cpp
--- a/Engine/Source/Editor/GameProjectGeneration/Private/Tests/GameProjectAutomationTests.cpp
+++ b/Engine/Source/Editor/GameProjectGeneration/Private/Tests/GameProjectAutomationTests.cpp
@@ -43,7 +43,7 @@ namespace GameProjectAutomationUtils
 		}
 
 
-		const FString Filename = ProjectName + TEXT(".") + IProjectManager::GetProjectFileExtension();
+		const FString Filename = ProjectName + TEXT(".") + FProjectDescriptor::GetExtension();
 		FString ProjectFilename = FPaths::Combine(*ProjectPath, *ProjectName, *Filename);
 		FPaths::MakePlatformFilename(ProjectFilename);
 
@@ -105,7 +105,7 @@ bool FBuildPromotionNewProjectCreateTest::RunTest(const FString& Parameters)
 		for (uint32 i = 2;; ++i)
 		{
 			const FString PossibleProjectName = FString::Printf(TEXT("%s%i"), *OldProjectName, i);
-			const FString PossibleProjectFilename = FString::Printf(TEXT("%s/%s/%s.%s"), *RootFolder, *PossibleProjectName, *PossibleProjectName, *IProjectManager::GetProjectFileExtension());
+			const FString PossibleProjectFilename = FString::Printf(TEXT("%s/%s/%s.%s"), *RootFolder, *PossibleProjectName, *PossibleProjectName, *FProjectDescriptor::GetExtension());
 			if (!FPaths::FileExists(PossibleProjectFilename))
 			{
 				DesiredProjectFilename = PossibleProjectFilename;
diff --git a/Engine/Source/Editor/UnrealEd/Private/Editor.cpp b/Engine/Source/Editor/UnrealEd/Private/Editor.cpp
--- a/Engine/Source/Editor/UnrealEd/Private/Editor.cpp
+++ b/Engine/Source/Editor/UnrealEd/Private/Editor.cpp
@@ -6471,7 +6471,7 @@ void UEditorEngine::UpdateRecentlyLoadedProjectFiles()
 		for ( int32 FileIdx = RecentlyOpenedProjectFiles.Num() - 1; FileIdx >= 0; --FileIdx )
 		{
 			const FString FileExtension = FPaths::GetExtension(RecentlyOpenedProjectFiles[FileIdx]);
-			if ( FileExtension != IProjectManager::GetProjectFileExtension() )
+			if ( FileExtension != FProjectDescriptor::GetExtension() )
 			{
 				RecentlyOpenedProjectFiles.RemoveAt(FileIdx, 1);
 			}
diff --git a/Engine/Source/Editor/UnrealEd/Private/UnrealEdMisc.cpp b/Engine/Source/Editor/UnrealEd/Private/UnrealEdMisc.cpp
--- a/Engine/Source/Editor/UnrealEd/Private/UnrealEdMisc.cpp
+++ b/Engine/Source/Editor/UnrealEd/Private/UnrealEdMisc.cpp
@@ -1188,7 +1188,7 @@ void FUnrealEdMisc::SwitchProject(const FString& GameOrProjectFileName, bool bWa
 		return;
 	}
 
-	const bool bIsProjectFileName = FPaths::GetExtension(GameOrProjectFileName) == IProjectManager::GetProjectFileExtension();
+	const bool bIsProjectFileName = FPaths::GetExtension(GameOrProjectFileName) == FProjectDescriptor::GetExtension();
 
 	bool bSwitch = true;
 
diff --git a/Engine/Source/Runtime/Launch/Private/LaunchEngineLoop.cpp b/Engine/Source/Runtime/Launch/Private/LaunchEngineLoop.cpp
--- a/Engine/Source/Runtime/Launch/Private/LaunchEngineLoop.cpp
+++ b/Engine/Source/Runtime/Launch/Private/LaunchEngineLoop.cpp
@@ -197,7 +197,7 @@ bool ParseGameProjectFromCommandLine(const TCHAR* InCmdLine, FString& OutProject
 	{
 		// The first command line argument could be the project file if it exists or the game name if not launching with a project file
 		const FString ProjectFilePath = FString(FirstCommandLineToken);
-		if ( FPaths::GetExtension(ProjectFilePath) == IProjectManager::GetProjectFileExtension() )
+		if ( FPaths::GetExtension(ProjectFilePath) == FProjectDescriptor::GetExtension() )
 		{
 			OutProjectFilePath = FirstCommandLineToken;
 			// Here we derive the game name from the project file
@@ -209,7 +209,7 @@ bool ParseGameProjectFromCommandLine(const TCHAR* InCmdLine, FString& OutProject
 			// Full game name is assumed to be the first token
 			OutGameName = MoveTemp(FirstCommandLineToken);
 			// Derive the project path from the game name. All games must have a uproject file, even if they are in the root folder.
-			OutProjectFilePath = FPaths::Combine(*FPaths::RootDir(), *OutGameName, *FString(OutGameName + TEXT(".") + IProjectManager::GetProjectFileExtension()));
+			OutProjectFilePath = FPaths::Combine(*FPaths::RootDir(), *OutGameName, *FString(OutGameName + TEXT(".") + FProjectDescriptor::GetExtension()));
 			return true;
 		}
 	}
@@ -244,7 +244,7 @@ bool LaunchSetGameName(const TCHAR *InCmdLine)
 			// Check it's not UE4Game, otherwise assume a uproject file relative to the game project directory
 			if (LocalGameName != TEXT("UE4Game"))
 			{
-				ProjFilePath = FPaths::Combine(TEXT(".."), TEXT(".."), TEXT(".."), *LocalGameName, *FString(LocalGameName + TEXT(".") + IProjectManager::GetProjectFileExtension()));
+				ProjFilePath = FPaths::Combine(TEXT(".."), TEXT(".."), TEXT(".."), *LocalGameName, *FString(LocalGameName + TEXT(".") + FProjectDescriptor::GetExtension()));
 				FPaths::SetProjectFilePath(ProjFilePath);
 			}
 		}
@@ -907,7 +907,7 @@ int32 FEngineLoop::PreInit( const TCHAR* CmdLine )
 	if ( !GIsGameAgnosticExe && FApp::HasGameName() && !FPaths::IsProjectFilePathSet() )
 	{
 		// If we are using a non-agnostic exe where a name was specified but we did not specify a project path. Assemble one based on the game name.
-		const FString ProjectFilePath = FPaths::Combine(*FPaths::GameDir(), *FString::Printf(TEXT("%s.%s"), FApp::GetGameName(), *IProjectManager::GetProjectFileExtension()));
+		const FString ProjectFilePath = FPaths::Combine(*FPaths::GameDir(), *FString::Printf(TEXT("%s.%s"), FApp::GetGameName(), *FProjectDescriptor::GetExtension()));
 		FPaths::SetProjectFilePath(ProjectFilePath);
 	}
 
diff --git a/Engine/Source/Runtime/NetworkFileSystem/Private/NetworkFileServerConnection.cpp b/Engine/Source/Runtime/NetworkFileSystem/Private/NetworkFileServerConnection.cpp
--- a/Engine/Source/Runtime/NetworkFileSystem/Private/NetworkFileServerConnection.cpp
+++ b/Engine/Source/Runtime/NetworkFileSystem/Private/NetworkFileServerConnection.cpp
@@ -725,7 +725,7 @@ bool FNetworkFileServerClientConnection::ProcessGetFileList( FArchive& In, FArch
 	}
 	else
 	{
-		if (FPaths::GetExtension(GameName) == IProjectManager::GetProjectFileExtension())
+		if (FPaths::GetExtension(GameName) == FProjectDescriptor::GetExtension())
 		{
 			SandboxDirectory = FPaths::Combine(*FPaths::GetPath(GameName), TEXT("Saved"), TEXT("Cooked"), *ConnectedPlatformName);
 		}
diff --git a/Engine/Source/Runtime/Projects/Private/ProjectDescriptor.cpp b/Engine/Source/Runtime/Projects/Private/ProjectDescriptor.cpp
--- a/Engine/Source/Runtime/Projects/Private/ProjectDescriptor.cpp
+++ b/Engine/Source/Runtime/Projects/Private/ProjectDescriptor.cpp
@@ -171,4 +171,10 @@ void FProjectDescriptor::Write(TJsonWriter<>& Writer) const
 	Writer.WriteObjectEnd();
 }
 
+FString FProjectDescriptor::GetExtension()
+{
+	static const FString ProjectExtension(TEXT("uproject"));
+	return ProjectExtension;
+}
+
 #undef LOCTEXT_NAMESPACE
diff --git a/Engine/Source/Runtime/Projects/Private/ProjectManager.cpp b/Engine/Source/Runtime/Projects/Private/ProjectManager.cpp
--- a/Engine/Source/Runtime/Projects/Private/ProjectManager.cpp
+++ b/Engine/Source/Runtime/Projects/Private/ProjectManager.cpp
@@ -107,12 +107,6 @@ const FString& FProjectManager::GetAutoLoadProjectFileName()
 	return RecentProjectFileName;
 }
 
-const FString& FProjectManager::NonStaticGetProjectFileExtension()
-{
-	static FString GameProjectFileExtension(TEXT("uproject"));
-	return GameProjectFileExtension;
-}
-
 bool FProjectManager::GenerateNewProjectFile(const FString& NewProjectFilename, const TArray<FString>& StartupModuleNames, const FString& EngineIdentifier, FText& OutFailReason)
 {
 	FProjectDescriptor Descriptor;
diff --git a/Engine/Source/Runtime/Projects/Private/ProjectManager.h b/Engine/Source/Runtime/Projects/Private/ProjectManager.h
--- a/Engine/Source/Runtime/Projects/Private/ProjectManager.h
+++ b/Engine/Source/Runtime/Projects/Private/ProjectManager.h
@@ -18,7 +18,6 @@ public:
 	virtual bool LoadModulesForProject( const ELoadingPhase::Type LoadingPhase ) override;
 	virtual bool AreProjectModulesUpToDate( ) override;
 	virtual const FString& GetAutoLoadProjectFileName() override;
-	virtual const FString& NonStaticGetProjectFileExtension() override;
 	virtual bool GenerateNewProjectFile(const FString& NewProjectFilename, const TArray<FString>& ModuleNames, const FString& EngineIdentifier, FText& OutFailReason) override;
 	virtual bool SignSampleProject(const FString& FilePath, const FString& Category, FText& OutFailReason) override;
 	virtual bool QueryStatusForProject(const FString& FilePath, FProjectStatus& OutProjectStatus) const override;
diff --git a/Engine/Source/Runtime/Projects/Public/Interfaces/IProjectManager.h b/Engine/Source/Runtime/Projects/Public/Interfaces/IProjectManager.h
--- a/Engine/Source/Runtime/Projects/Public/Interfaces/IProjectManager.h
+++ b/Engine/Source/Runtime/Projects/Public/Interfaces/IProjectManager.h
@@ -119,13 +119,6 @@ public:
 	 */
 	virtual const FString& GetAutoLoadProjectFileName() = 0;
 
-	/**
-	 * Gets the extension for game project files.
-	 *
-	 * @return File extension.
-	 */
-	virtual const FString& NonStaticGetProjectFileExtension() = 0;
-
 	/**
 	 * Generates a new project file and saves it to disk at the specified location
 	 * 
@@ -227,7 +220,4 @@ public:
 	 * @return	True if the application needs to be restarted.
 	 */
 	virtual bool IsRestartRequired() const = 0;
-
-	/** Helper functions to reduce the syntax complexity of commonly used functions */
-	static const FString& GetProjectFileExtension() { return Get().NonStaticGetProjectFileExtension(); }
 };
\ No newline at end of file
diff --git a/Engine/Source/Runtime/Projects/Public/ProjectDescriptor.h b/Engine/Source/Runtime/Projects/Public/ProjectDescriptor.h
--- a/Engine/Source/Runtime/Projects/Public/ProjectDescriptor.h
+++ b/Engine/Source/Runtime/Projects/Public/ProjectDescriptor.h
@@ -77,4 +77,7 @@ struct PROJECTS_API FProjectDescriptor
 
 	/** Writes the descriptor to the given JSON object */
 	void Write(TJsonWriter<>& Writer) const;
+
+	/** Returns the extension used for project descriptors (uproject) */
+	static FString GetExtension();
 };
