Added HardwareTargeting module and tidied up project browser dialogs

Also removed tabs from the project browser dialogs when they're not necessary. Fixed up some incorrect icon sizes to work with the new layout.

diff --git a/Engine/Source/Editor/DetailCustomizations/DetailCustomizations.Build.cs b/Engine/Source/Editor/DetailCustomizations/DetailCustomizations.Build.cs
--- a/Engine/Source/Editor/DetailCustomizations/DetailCustomizations.Build.cs
+++ b/Engine/Source/Editor/DetailCustomizations/DetailCustomizations.Build.cs
@@ -40,6 +40,7 @@ public class DetailCustomizations : ModuleRules
 				"RHI",
                 "Media",
                 "MediaAssets",
+                "HardwareTargeting",
 			}
 		);
 
diff --git a/Engine/Source/Editor/DetailCustomizations/Private/DetailCustomizations.cpp b/Engine/Source/Editor/DetailCustomizations/Private/DetailCustomizations.cpp
--- a/Engine/Source/Editor/DetailCustomizations/Private/DetailCustomizations.cpp
+++ b/Engine/Source/Editor/DetailCustomizations/Private/DetailCustomizations.cpp
@@ -71,6 +71,7 @@
 #include "CollisionProfileDetails.h"
 #include "PhysicsSettingsDetails.h"
 #include "GeneralProjectSettingsDetails.h"
+#include "HardwareTargetingSettingsDetails.h"
 #include "WindowsTargetSettingsDetails.h"
 #include "MoviePlayerSettingsDetails.h"
 #include "SourceCodeAccessSettingsDetails.h"
@@ -201,6 +202,7 @@ void FDetailCustomizationsModule::RegisterObjectCustomizations()
 	RegisterCustomClassLayout("AmbientSound", FOnGetDetailCustomizationInstance::CreateStatic(&FAmbientSoundDetails::MakeInstance));
 	RegisterCustomClassLayout("WorldSettings", FOnGetDetailCustomizationInstance::CreateStatic(&FWorldSettingsDetails::MakeInstance));
 	RegisterCustomClassLayout("GeneralProjectSettings", FOnGetDetailCustomizationInstance::CreateStatic(&FGeneralProjectSettingsDetails::MakeInstance));
+	RegisterCustomClassLayout("HardwareTargetingSettings", FOnGetDetailCustomizationInstance::CreateStatic(&FHardwareTargetingSettingsDetails::MakeInstance));
 	RegisterCustomClassLayout("DocumentationActor", FOnGetDetailCustomizationInstance::CreateStatic(&FDocumentationActorDetails::MakeInstance));
 
 	//@TODO: A2REMOVAL: Rename FSkeletalControlNodeDetails to something more generic
diff --git a/Engine/Source/Editor/DetailCustomizations/Private/HardwareTargetingSettingsDetails.cpp b/Engine/Source/Editor/DetailCustomizations/Private/HardwareTargetingSettingsDetails.cpp
new file mode 100644
--- /dev/null
+++ b/Engine/Source/Editor/DetailCustomizations/Private/HardwareTargetingSettingsDetails.cpp
@@ -0,0 +1,122 @@
+// Copyright 1998-2014 Epic Games, Inc. All Rights Reserved.
+
+#include "DetailCustomizationsPrivatePCH.h"
+
+#include "HardwareTargetingSettingsDetails.h"
+#include "HardwareTargetingModule.h"
+
+#define LOCTEXT_NAMESPACE "FHardwareTargetingSettingsDetails"
+
+TSharedRef<IDetailCustomization> FHardwareTargetingSettingsDetails::MakeInstance()
+{
+	return MakeShareable(new FHardwareTargetingSettingsDetails);
+}
+
+void FHardwareTargetingSettingsDetails::CustomizeDetails(IDetailLayoutBuilder& DetailBuilder)
+{	
+	IDetailCategoryBuilder& HardwareTargetingCategory = DetailBuilder.EditCategory(TEXT("Hardware Targeting"));
+	
+	IHardwareTargetingModule& HardwareTargeting = IHardwareTargetingModule::Get();
+	HardwareTargetingCategory.AddCustomRow(FString())
+	[
+		SNew(SBorder)
+		.BorderImage(FEditorStyle::GetBrush("NoBorder"))
+		.Padding(0)
+		.Visibility_Static([]()->EVisibility{
+			auto* Settings = GetMutableDefault<UHardwareTargetingSettings>();
+			return Settings->HasPendingChanges() ? EVisibility::Visible : EVisibility::Collapsed;
+		})
+		[
+			SNew(SHorizontalBox)
+
+			+ SHorizontalBox::Slot()
+			.AutoWidth()
+			.VAlign(VAlign_Center)
+			[
+				SNew(STextBlock)
+				.Text(LOCTEXT("RestartMessage", "These changes will be applied when this project is reopened."))
+			]
+
+			+ SHorizontalBox::Slot()
+			.AutoWidth()
+			.Padding(6)
+			[
+				SNew(SButton)
+				.Text(LOCTEXT("ApplyNow", "Apply Now"))
+				.OnClicked_Static([](){
+				
+					IHardwareTargetingModule& Module = IHardwareTargetingModule::Get();
+					Module.ApplyHardwareTargetingSettings();
+
+					return FReply::Handled();
+				})
+			]
+		]
+	];
+
+
+	auto PropertyName = GET_MEMBER_NAME_CHECKED(UHardwareTargetingSettings, TargetedHardwareClass);
+	DetailBuilder.HideProperty(PropertyName);
+
+	TSharedRef<IPropertyHandle> Property = DetailBuilder.GetProperty(PropertyName);
+	HardwareTargetingCategory.AddCustomRow(LOCTEXT("HardwareTargetingOption", "Targeted Hardware:").ToString())
+	.NameContent()
+	[
+		SNew(STextBlock)
+		.Text(FText::FromString(Property->GetPropertyDisplayName()))
+		.Font(DetailBuilder.GetDetailFont())
+	]
+	.ValueContent()
+	[
+		SNew(SHorizontalBox)
+
+		+ SHorizontalBox::Slot()
+		.AutoWidth()
+		[
+			HardwareTargeting.MakeHardwareClassTargetCombo(
+				FOnHardwareClassChanged::CreateStatic([](EHardwareClass::Type NewValue, TSharedRef<IPropertyHandle> Property){
+					Property->SetValue(uint8(NewValue));
+				}, Property),
+				TAttribute<EHardwareClass::Type>::Create(TAttribute<EHardwareClass::Type>::FGetter::CreateStatic([](TSharedRef<IPropertyHandle> Property){
+					uint8 Value = 0;
+					Property->GetValue(Value);
+					return EHardwareClass::Type(Value);
+				}, Property))
+			)
+		]
+	];
+
+	PropertyName = GET_MEMBER_NAME_CHECKED(UHardwareTargetingSettings, DefaultGraphicsPerformance);
+	DetailBuilder.HideProperty(PropertyName);
+
+	Property = DetailBuilder.GetProperty(PropertyName);
+	HardwareTargetingCategory.AddCustomRow(LOCTEXT("GraphicsPreferenceOption", "Graphics Preference:").ToString())
+	.NameContent()
+	[
+		SNew(STextBlock)
+		.Text(FText::FromString(Property->GetPropertyDisplayName()))
+		.Font(DetailBuilder.GetDetailFont())
+	]
+	.ValueContent()
+	[
+		SNew(SHorizontalBox)
+
+		+ SHorizontalBox::Slot()
+		.AutoWidth()
+		[
+			HardwareTargeting.MakeGraphicsPresetTargetCombo(
+				FOnGraphicsPresetChanged::CreateStatic([](EGraphicsPreset::Type NewValue, TSharedRef<IPropertyHandle> Property){
+					Property->SetValue(uint8(NewValue));
+				}, Property),
+				TAttribute<EGraphicsPreset::Type>::Create(TAttribute<EGraphicsPreset::Type>::FGetter::CreateStatic([](TSharedRef<IPropertyHandle> Property){
+					uint8 Value = 0;
+					Property->GetValue(Value);
+					return EGraphicsPreset::Type(Value);
+				}, Property))
+			)
+		]
+	];
+
+}
+
+#undef LOCTEXT_NAMESPACE
diff --git a/Engine/Source/Editor/DetailCustomizations/Private/HardwareTargetingSettingsDetails.h b/Engine/Source/Editor/DetailCustomizations/Private/HardwareTargetingSettingsDetails.h
new file mode 100644
--- /dev/null
+++ b/Engine/Source/Editor/DetailCustomizations/Private/HardwareTargetingSettingsDetails.h
@@ -0,0 +1,17 @@
+// Copyright 1998-2014 Epic Games, Inc. All Rights Reserved.
+
+#pragma once
+
+#include "HardwareTargetingSettings.h"
+
+class FHardwareTargetingSettingsDetails : public IDetailCustomization
+{
+public:
+	/** Makes a new instance of this detail layout class for a specific detail view requesting it */
+	static TSharedRef<IDetailCustomization> MakeInstance();
+
+	// IDetailCustomization interface
+	virtual void CustomizeDetails(IDetailLayoutBuilder& DetailBuilder) override;
+	// End of IDetailCustomization interface
+};
+
diff --git a/Engine/Source/Editor/EditorStyle/Private/SlateEditorStyle.cpp b/Engine/Source/Editor/EditorStyle/Private/SlateEditorStyle.cpp
--- a/Engine/Source/Editor/EditorStyle/Private/SlateEditorStyle.cpp
+++ b/Engine/Source/Editor/EditorStyle/Private/SlateEditorStyle.cpp
@@ -1559,11 +1559,11 @@ void FSlateEditorStyle::FStyle::SetupGeneralStyles()
 			.SetSelectedTextColor( InvertedForeground )
 			);
 
-		Set( "GameProjectDialog.DefaultGameThumbnail", new IMAGE_BRUSH( "GameProjectDialog/default_game_thumbnail", Icon128x128 ) );
+		Set( "GameProjectDialog.DefaultGameThumbnail", new IMAGE_BRUSH( "GameProjectDialog/default_game_thumbnail_128x", Icon128x128 ) );
+		Set( "GameProjectDialog.DefaultGameThumbnail.Small", new IMAGE_BRUSH( "GameProjectDialog/default_game_thumbnail", Icon128x128 ) );
 		Set( "GameProjectDialog.BlankProjectThumbnail", new IMAGE_BRUSH( "GameProjectDialog/blank_project_thumbnail", Icon128x128 ) );
-		Set( "GameProjectDialog.BlankProjectPreview", new IMAGE_BRUSH( "GameProjectDialog/blank_project_preview", FVector2D(512, 256) ) );
+		Set( "GameProjectDialog.BlankProjectPreview", new IMAGE_BRUSH( "GameProjectDialog/blank_project_preview", FVector2D(400, 200) ) );
 		Set( "GameProjectDialog.BasicCodeThumbnail", new IMAGE_BRUSH( "GameProjectDialog/basic_code_thumbnail", Icon128x128 ) );
-		Set( "GameProjectDialog.BasicCodePreview", new IMAGE_BRUSH( "GameProjectDialog/basic_code_preview", FVector2D(512, 256) ) );
 		Set( "GameProjectDialog.CodeIcon", new IMAGE_BRUSH( "GameProjectDialog/feature_code_32x", FVector2D(32,32) ) );
 		Set( "GameProjectDialog.CodeImage", new IMAGE_BRUSH( "GameProjectDialog/feature_code", FVector2D(96,96) ) );
 		Set( "GameProjectDialog.BlueprintIcon", new IMAGE_BRUSH( "GameProjectDialog/feature_blueprint_32x", FVector2D(32,32) ) );
@@ -1594,6 +1594,9 @@ void FSlateEditorStyle::FStyle::SetupGeneralStyles()
 		Set( "GameProjectDialog.FolderIconOpen", new IMAGE_BRUSH( "Icons/FolderOpen", FVector2D(18, 16) ) );
 		Set( "GameProjectDialog.ProjectFileIcon", new IMAGE_BRUSH( "Icons/doc_16x", FVector2D(18, 16) ) );
 
+		Set( "GameProjectDialog.IncludeStarterContent", new IMAGE_BRUSH( "/GameProjectDialog/IncludeStarterContent", FVector2D(64, 64) ) );
+		Set( "GameProjectDialog.NoStarterContent", new IMAGE_BRUSH( "/GameProjectDialog/NoStarterContent", FVector2D(64, 64) ) );
+
 		Set( "FilePath.FolderButton",
 			FButtonStyle(HoverHintOnly)
 			.SetNormal(  BOX_BRUSH( "Common/ButtonHoverHint", FMargin(4/16.0f), FLinearColor(0.5f,0.5f,0.5f) ) )
@@ -1659,6 +1662,14 @@ void FSlateEditorStyle::FStyle::SetupGeneralStyles()
 		);
 	}
 
+	// Hardware Targeting
+	{
+		Set( "HardwareTargeting.MobilePlatform", new IMAGE_BRUSH( "/Icons/HardwareTargeting/Mobile", FVector2D(64, 64) ) );
+		Set( "HardwareTargeting.DesktopPlatform", new IMAGE_BRUSH( "/Icons/HardwareTargeting/Desktop", FVector2D(64, 64) ) );
+
+		Set( "HardwareTargeting.MaximumQuality", new IMAGE_BRUSH( "/Icons/HardwareTargeting/MaximumQuality", FVector2D(64, 64) ) );
+		Set( "HardwareTargeting.ScalableQuality", new IMAGE_BRUSH( "/Icons/HardwareTargeting/ScalableQuality", FVector2D(64, 64) ) );
+	}
 
 #endif // WITH_EDITOR
 
@@ -4947,6 +4958,11 @@ void FSlateEditorStyle::FStyle::SetupToolkitStyles()
 			.SetFont(TTF_CORE_FONT("Fonts/Roboto-Regular", 24))
 			.SetShadowOffset(FVector2D(0, 1)));
 
+		Set("ProjectBrowser.Toolbar.Text", FTextBlockStyle(NormalText)
+			.SetFont(TTF_CORE_FONT("Fonts/Roboto-Bold", 12))
+			.SetShadowColorAndOpacity(FLinearColor(0,0,0,0.5f))
+			.SetShadowOffset(FVector2D(0, 1)));
+
 		Set("ProjectBrowser.VersionOverlayText", FTextBlockStyle(NormalText)
 			.SetFont(TTF_CORE_FONT("Fonts/Roboto-Bold", 14))
 			.SetShadowOffset(FVector2D(0, 1)));
diff --git a/Engine/Source/Editor/GameProjectGeneration/GameProjectGeneration.Build.cs b/Engine/Source/Editor/GameProjectGeneration/GameProjectGeneration.Build.cs
--- a/Engine/Source/Editor/GameProjectGeneration/GameProjectGeneration.Build.cs
+++ b/Engine/Source/Editor/GameProjectGeneration/GameProjectGeneration.Build.cs
@@ -31,7 +31,8 @@ public class GameProjectGeneration : ModuleRules
                 "SourceControl",
  				"TargetPlatform",
 				"UnrealEd",
-				"DesktopPlatform"
+				"DesktopPlatform",
+                "HardwareTargeting"
 			}
 		);
 
diff --git a/Engine/Source/Editor/GameProjectGeneration/Private/GameProjectUtils.cpp b/Engine/Source/Editor/GameProjectGeneration/Private/GameProjectUtils.cpp
--- a/Engine/Source/Editor/GameProjectGeneration/Private/GameProjectUtils.cpp
+++ b/Engine/Source/Editor/GameProjectGeneration/Private/GameProjectUtils.cpp
@@ -318,7 +318,7 @@ bool GameProjectUtils::IsValidProjectFileForCreation(const FString& ProjectFile,
 	{
 		FFormatNamedArguments Args;
 		Args.Add( TEXT("MaxProjectPathLength"), MaxProjectPathLength );
-		OutFailReason = FText::Format( LOCTEXT( "ProjectPathTooLong", "A projects path must not be longer than {MaxProjectPathLength} characters." ), Args );
+		OutFailReason = FText::Format( LOCTEXT( "ProjectPathTooLong", "A project's path must not be longer than {MaxProjectPathLength} characters." ), Args );
 		return false;
 	}
 
@@ -352,17 +352,13 @@ bool GameProjectUtils::IsValidProjectFileForCreation(const FString& ProjectFile,
 
 	if ( ProjectFileExists(ProjectFile) )
 	{
-		FFormatNamedArguments Args;
-		Args.Add( TEXT("ProjectFile"), FText::FromString( ProjectFile ) );
-		OutFailReason = FText::Format( LOCTEXT( "ProjectFileAlreadyExists", "{ProjectFile} already exists." ), Args );
+		OutFailReason = LOCTEXT( "ProjectFileAlreadyExists", "This project file already exists." );
 		return false;
 	}
 
 	if ( FPaths::ConvertRelativePathToFull(FPaths::GetPath(ProjectFile)).StartsWith( FPaths::ConvertRelativePathToFull(FPaths::EngineDir())) )
 	{
-		FFormatNamedArguments Args;
-		Args.Add( TEXT("ProjectFile"), FText::FromString( ProjectFile ) );
-		OutFailReason = FText::Format( LOCTEXT( "ProjectFileCannotBeUnderEngineFolder", "{ProjectFile} cannot be saved under the Engine folder.  Create the project in a different directory." ), Args );
+		OutFailReason = LOCTEXT( "ProjectFileCannotBeUnderEngineFolder", "Project cannot be saved under the Engine folder. Please choose a different directory." );
 		return false;
 	}
 
@@ -370,8 +366,7 @@ bool GameProjectUtils::IsValidProjectFileForCreation(const FString& ProjectFile,
 	{
 		FFormatNamedArguments Args;
 		Args.Add( TEXT("ProjectFileExtension"), FText::FromString( FProjectDescriptor::GetExtension() ) );
-		Args.Add( TEXT("ProjectFilePath"), FText::FromString( FPaths::GetPath(ProjectFile) ) );
-		OutFailReason = FText::Format( LOCTEXT( "AProjectFileAlreadyExistsAtLoction", "Another .{ProjectFileExtension} file already exists in {ProjectFilePath}" ), Args );
+		OutFailReason = FText::Format( LOCTEXT( "AProjectFileAlreadyExistsAtLoction", "Another .{ProjectFileExtension} file already exists in the specified folder" ), Args );
 		return false;
 	}
 
@@ -404,7 +399,7 @@ bool GameProjectUtils::OpenProject(const FString& ProjectFile, FText& OutFailRea
 	{
 		FFormatNamedArguments Args;
 		Args.Add( TEXT("MaxProjectPathLength"), MaxProjectPathLength );
-		OutFailReason = FText::Format( LOCTEXT( "ProjectPathTooLong", "A projects path must not be longer than {MaxProjectPathLength} characters." ), Args );
+		OutFailReason = FText::Format( LOCTEXT( "ProjectPathTooLong", "A project's path must not be longer than {MaxProjectPathLength} characters." ), Args );
 		return false;
 	}
 
@@ -556,9 +551,9 @@ bool GameProjectUtils::CopyStarterContent(const FString& DestProjectFolder, FTex
 }
 
 
-bool GameProjectUtils::CreateProject(const FString& NewProjectFile, const FString& TemplateFile, bool bShouldGenerateCode, bool bCopyStarterContent, FText& OutFailReason)
+bool GameProjectUtils::CreateProject(const FProjectInformation& InProjectInfo, FText& OutFailReason)
 {
-	if ( !IsValidProjectFileForCreation(NewProjectFile, OutFailReason) )
+	if ( !IsValidProjectFileForCreation(InProjectInfo.ProjectFilename, OutFailReason) )
 	{
 		return false;
 	}
@@ -568,22 +563,22 @@ bool GameProjectUtils::CreateProject(const FString& NewProjectFile, const FStrin
 
 	bool bProjectCreationSuccessful = false;
 	FString TemplateName;
-	if ( TemplateFile.IsEmpty() )
+	if ( InProjectInfo.TemplateFile.IsEmpty() )
 	{
-		bProjectCreationSuccessful = GenerateProjectFromScratch(NewProjectFile, bShouldGenerateCode, bCopyStarterContent, OutFailReason);
-		TemplateName = bShouldGenerateCode ? TEXT("Basic Code") : TEXT("Blank");
+		bProjectCreationSuccessful = GenerateProjectFromScratch(InProjectInfo, OutFailReason);
+		TemplateName = InProjectInfo.bShouldGenerateCode ? TEXT("Basic Code") : TEXT("Blank");
 	}
 	else
 	{
-		bProjectCreationSuccessful = CreateProjectFromTemplate(NewProjectFile, TemplateFile, bShouldGenerateCode, bCopyStarterContent, OutFailReason);
-		TemplateName = FPaths::GetBaseFilename(TemplateFile);
+		bProjectCreationSuccessful = CreateProjectFromTemplate(InProjectInfo, OutFailReason);
+		TemplateName = FPaths::GetBaseFilename(InProjectInfo.TemplateFile);
 	}
 
 	if( FEngineAnalytics::IsAvailable() )
 	{
 		TArray<FAnalyticsEventAttribute> EventAttributes;
 		EventAttributes.Add(FAnalyticsEventAttribute(TEXT("Template"), TemplateName));
-		EventAttributes.Add(FAnalyticsEventAttribute(TEXT("ProjectType"), bShouldGenerateCode ? TEXT("C++ Code") : TEXT("Content Only")));
+		EventAttributes.Add(FAnalyticsEventAttribute(TEXT("ProjectType"), InProjectInfo.bShouldGenerateCode ? TEXT("C++ Code") : TEXT("Content Only")));
 		EventAttributes.Add(FAnalyticsEventAttribute(TEXT("Outcome"), bProjectCreationSuccessful ? TEXT("Successful") : TEXT("Failed")));
 
 		FEngineAnalytics::GetProvider().RecordEvent( TEXT( "Editor.NewProject.ProjectCreated" ), EventAttributes );
@@ -816,14 +811,14 @@ UTemplateProjectDefs* GameProjectUtils::LoadTemplateDefs(const FString& ProjectD
 	return TemplateDefs;
 }
 
-bool GameProjectUtils::GenerateProjectFromScratch(const FString& NewProjectFile, bool bShouldGenerateCode, bool bCopyStarterContent, FText& OutFailReason)
+bool GameProjectUtils::GenerateProjectFromScratch(const FProjectInformation& InProjectInfo, FText& OutFailReason)
 {
-	const FString NewProjectFolder = FPaths::GetPath(NewProjectFile);
-	const FString NewProjectName = FPaths::GetBaseFilename(NewProjectFile);
+	const FString NewProjectFolder = FPaths::GetPath(InProjectInfo.ProjectFilename);
+	const FString NewProjectName = FPaths::GetBaseFilename(InProjectInfo.ProjectFilename);
 	TArray<FString> CreatedFiles;
 
 	// Generate config files
-	if (!GenerateConfigFiles(NewProjectFolder, NewProjectName, bShouldGenerateCode, bCopyStarterContent, CreatedFiles, OutFailReason))
+	if (!GenerateConfigFiles(InProjectInfo, CreatedFiles, OutFailReason))
 	{
 		DeleteCreatedFiles(NewProjectFolder, CreatedFiles);
 		return false;
@@ -841,7 +836,7 @@ bool GameProjectUtils::GenerateProjectFromScratch(const FString& NewProjectFile,
 	}
 
 	TArray<FString> StartupModuleNames;
-	if ( bShouldGenerateCode )
+	if ( InProjectInfo.bShouldGenerateCode )
 	{
 		// Generate basic source code files
 		if ( !GenerateBasicSourceCode(NewProjectFolder / TEXT("Source"), NewProjectName, StartupModuleNames, CreatedFiles, OutFailReason) )
@@ -861,9 +856,9 @@ bool GameProjectUtils::GenerateProjectFromScratch(const FString& NewProjectFile,
 	// Generate the project file
 	{
 		FText LocalFailReason;
-		if (IProjectManager::Get().GenerateNewProjectFile(NewProjectFile, StartupModuleNames, FDesktopPlatformModule::Get()->GetCurrentEngineIdentifier(), LocalFailReason))
+		if (IProjectManager::Get().GenerateNewProjectFile(InProjectInfo.ProjectFilename, StartupModuleNames, FDesktopPlatformModule::Get()->GetCurrentEngineIdentifier(), LocalFailReason))
 		{
-			CreatedFiles.Add(NewProjectFile);
+			CreatedFiles.Add(InProjectInfo.ProjectFilename);
 		}
 		else
 		{
@@ -873,23 +868,23 @@ bool GameProjectUtils::GenerateProjectFromScratch(const FString& NewProjectFile,
 		}
 	}
 
-	if ( bShouldGenerateCode )
+	if ( InProjectInfo.bShouldGenerateCode )
 	{
 		// Generate project files
-		if ( !GenerateCodeProjectFiles(NewProjectFile, OutFailReason) )
+		if ( !GenerateCodeProjectFiles(InProjectInfo.ProjectFilename, OutFailReason) )
 		{
-			DeleteGeneratedProjectFiles(NewProjectFile);
+			DeleteGeneratedProjectFiles(InProjectInfo.ProjectFilename);
 			DeleteCreatedFiles(NewProjectFolder, CreatedFiles);
 			return false;
 		}
 	}
 
-	if (bCopyStarterContent)
+	if (InProjectInfo.bCopyStarterContent)
 	{
 		// Copy the starter content
 		if ( !CopyStarterContent(NewProjectFolder, OutFailReason) )
 		{
-			DeleteGeneratedProjectFiles(NewProjectFile);
+			DeleteGeneratedProjectFiles(InProjectInfo.ProjectFilename);
 			DeleteCreatedFiles(NewProjectFolder, CreatedFiles);
 			return false;
 		}
@@ -899,17 +894,17 @@ bool GameProjectUtils::GenerateProjectFromScratch(const FString& NewProjectFile,
 	return true;
 }
 
-bool GameProjectUtils::CreateProjectFromTemplate(const FString& NewProjectFile, const FString& TemplateFile, bool bShouldGenerateCode, bool bCopyStarterContent, FText& OutFailReason)
+bool GameProjectUtils::CreateProjectFromTemplate(const FProjectInformation& InProjectInfo, FText& OutFailReason)
 {
-	const FString ProjectName = FPaths::GetBaseFilename(NewProjectFile);
-	const FString TemplateName = FPaths::GetBaseFilename(TemplateFile);
-	const FString SrcFolder = FPaths::GetPath(TemplateFile);
-	const FString DestFolder = FPaths::GetPath(NewProjectFile);
+	const FString ProjectName = FPaths::GetBaseFilename(InProjectInfo.ProjectFilename);
+	const FString TemplateName = FPaths::GetBaseFilename(InProjectInfo.TemplateFile);
+	const FString SrcFolder = FPaths::GetPath(InProjectInfo.TemplateFile);
+	const FString DestFolder = FPaths::GetPath(InProjectInfo.ProjectFilename);
 
-	if ( !FPlatformFileManager::Get().GetPlatformFile().FileExists(*TemplateFile) )
+	if ( !FPlatformFileManager::Get().GetPlatformFile().FileExists(*InProjectInfo.TemplateFile) )
 	{
 		FFormatNamedArguments Args;
-		Args.Add( TEXT("TemplateFile"), FText::FromString( TemplateFile ) );
+		Args.Add( TEXT("TemplateFile"), FText::FromString( InProjectInfo.TemplateFile ) );
 		OutFailReason = FText::Format( LOCTEXT("InvalidTemplate_MissingProject", "Template project \"{TemplateFile}\" does not exist."), Args );
 		return false;
 	}
@@ -918,7 +913,7 @@ bool GameProjectUtils::CreateProjectFromTemplate(const FString& NewProjectFile,
 	if ( TemplateDefs == NULL )
 	{
 		FFormatNamedArguments Args;
-		Args.Add( TEXT("TemplateFile"), FText::FromString( FPaths::GetBaseFilename(TemplateFile) ) );
+		Args.Add( TEXT("TemplateFile"), FText::FromString( FPaths::GetBaseFilename(InProjectInfo.TemplateFile) ) );
 		Args.Add( TEXT("TemplateDefinesFile"), FText::FromString( GetTemplateDefsFilename() ) );
 		OutFailReason = FText::Format( LOCTEXT("InvalidTemplate_MissingDefs", "Template project \"{TemplateFile}\" does not have definitions file: '{TemplateDefinesFile}'."), Args );
 		return false;
@@ -1080,9 +1075,28 @@ bool GameProjectUtils::CreateProjectFromTemplate(const FString& NewProjectFile,
 		}
 	}
 
+	const FString ProjectConfigPath = DestFolder / TEXT("Config");
+
+	// Write out the hardware class target settings chosen for this project
+	{
+		const FString DefaultEditorIniFilename = ProjectConfigPath / TEXT("DefaultEditor.ini");
+
+		FString FileContents;
+		// Load the existing file - if it doesn't exist we create it
+		FFileHelper::LoadFileToString(FileContents, *DefaultEditorIniFilename);
+
+		FileContents += LINE_TERMINATOR;
+		FileContents += GetHardwareConfigString(InProjectInfo);
+
+		if ( !WriteOutputFile(DefaultEditorIniFilename, FileContents, OutFailReason) )
+		{
+			return false;
+		}
+	}
+
 	// Fixup specific ini values
 	TArray<FTemplateConfigValue> ConfigValuesToSet;
-	TemplateDefs->AddConfigValues(ConfigValuesToSet, TemplateName, ProjectName, bShouldGenerateCode);
+	TemplateDefs->AddConfigValues(ConfigValuesToSet, TemplateName, ProjectName, InProjectInfo.bShouldGenerateCode);
 	new (ConfigValuesToSet) FTemplateConfigValue(TEXT("DefaultGame.ini"), TEXT("/Script/EngineSettings.GeneralProjectSettings"), TEXT("ProjectID"), FGuid::NewGuid().ToString(), /*InShouldReplaceExistingValue=*/true);
 
 	// Add all classname fixups
@@ -1096,7 +1110,7 @@ bool GameProjectUtils::CreateProjectFromTemplate(const FString& NewProjectFile,
 	for ( auto ConfigIt = ConfigValuesToSet.CreateConstIterator(); ConfigIt; ++ConfigIt )
 	{
 		const FTemplateConfigValue& ConfigValue = *ConfigIt;
-		const FString IniFilename = DestFolder / TEXT("Config") / ConfigValue.ConfigFile;
+		const FString IniFilename = ProjectConfigPath / ConfigValue.ConfigFile;
 		bool bSuccessfullyProcessed = false;
 
 		TArray<FString> FileLines;
@@ -1204,7 +1218,7 @@ bool GameProjectUtils::CreateProjectFromTemplate(const FString& NewProjectFile,
 	{
 		// Load the source project
 		FProjectDescriptor Project;
-		if(!Project.Load(TemplateFile, OutFailReason))
+		if(!Project.Load(InProjectInfo.TemplateFile, OutFailReason))
 		{
 			DeleteCreatedFiles(DestFolder, CreatedFiles);
 			return false;
@@ -1215,8 +1229,8 @@ bool GameProjectUtils::CreateProjectFromTemplate(const FString& NewProjectFile,
 		Project.EpicSampleNameHash = 0;
 
 		// Fix up module names
-		const FString BaseSourceName = FPaths::GetBaseFilename(TemplateFile);
-		const FString BaseNewName = FPaths::GetBaseFilename(NewProjectFile);
+		const FString BaseSourceName = FPaths::GetBaseFilename(InProjectInfo.TemplateFile);
+		const FString BaseNewName = FPaths::GetBaseFilename(InProjectInfo.ProjectFilename);
 		for ( auto ModuleIt = Project.Modules.CreateIterator(); ModuleIt; ++ModuleIt )
 		{
 			FModuleDescriptor& ModuleInfo = *ModuleIt;
@@ -1224,17 +1238,17 @@ bool GameProjectUtils::CreateProjectFromTemplate(const FString& NewProjectFile,
 		}
 
 		// Save it to disk
-		if(!Project.Save(NewProjectFile, OutFailReason))
+		if(!Project.Save(InProjectInfo.ProjectFilename, OutFailReason))
 		{
 			DeleteCreatedFiles(DestFolder, CreatedFiles);
 			return false;
 		}
 
 		// Add it to the list of created files
-		CreatedFiles.Add(NewProjectFile);
+		CreatedFiles.Add(InProjectInfo.ProjectFilename);
 	}
 
-	if ( bShouldGenerateCode )
+	if ( InProjectInfo.bShouldGenerateCode )
 	{
 		// resource folder
 		const FString GameModuleSourcePath = DestFolder / TEXT("Source") / ProjectName;
@@ -1245,28 +1259,28 @@ bool GameProjectUtils::CreateProjectFromTemplate(const FString& NewProjectFile,
 		}
 
 		// Generate project files
-		if ( !GenerateCodeProjectFiles(NewProjectFile, OutFailReason) )
+		if ( !GenerateCodeProjectFiles(InProjectInfo.ProjectFilename, OutFailReason) )
 		{
-			DeleteGeneratedProjectFiles(NewProjectFile);
+			DeleteGeneratedProjectFiles(InProjectInfo.ProjectFilename);
 			DeleteCreatedFiles(DestFolder, CreatedFiles);
 			return false;
 		}
 	}
 
-	if (bCopyStarterContent)
+	if (InProjectInfo.bCopyStarterContent)
 	{
 		// Copy the starter content
 		if ( !CopyStarterContent(DestFolder, OutFailReason) )
 		{
-			DeleteGeneratedProjectFiles(NewProjectFile);
+			DeleteGeneratedProjectFiles(InProjectInfo.ProjectFilename);
 			DeleteCreatedFiles(DestFolder, CreatedFiles);
 			return false;
 		}
 	}
 
-	if (!TemplateDefs->PostGenerateProject(DestFolder, SrcFolder, NewProjectFile, TemplateFile, bShouldGenerateCode, OutFailReason))
+	if (!TemplateDefs->PostGenerateProject(DestFolder, SrcFolder, InProjectInfo.ProjectFilename, InProjectInfo.TemplateFile, InProjectInfo.bShouldGenerateCode, OutFailReason))
 	{
-		DeleteGeneratedProjectFiles(NewProjectFile);
+		DeleteGeneratedProjectFiles(InProjectInfo.ProjectFilename);
 		DeleteCreatedFiles(DestFolder, CreatedFiles);
 		return false;
 	}
@@ -1395,9 +1409,24 @@ void GameProjectUtils::DeleteGeneratedBuildFiles(const FString& NewProjectFolder
 	}
 }
 
-bool GameProjectUtils::GenerateConfigFiles(const FString& NewProjectPath, const FString& NewProjectName, bool bShouldGenerateCode, bool bCopyStarterContent, TArray<FString>& OutCreatedFiles, FText& OutFailReason)
+FString GameProjectUtils::GetHardwareConfigString(const FProjectInformation& InProjectInfo)
+{
+	FString HardwareTargeting;
+	
+	HardwareTargeting += TEXT("[/Script/HardwareTargeting.HardwareTargetingSettings]") LINE_TERMINATOR;
+	HardwareTargeting += FString::Printf(TEXT("TargetedHardwareClass=%d") LINE_TERMINATOR, int32(InProjectInfo.TargetedHardware));
+	HardwareTargeting += FString::Printf(TEXT("DefaultGraphicsPerformance=%d") LINE_TERMINATOR, int32(InProjectInfo.DefaultGraphicsPerformance));
+	HardwareTargeting += LINE_TERMINATOR;
+
+	return HardwareTargeting;
+}
+
+bool GameProjectUtils::GenerateConfigFiles(const FProjectInformation& InProjectInfo, TArray<FString>& OutCreatedFiles, FText& OutFailReason)
 {
-	FString ProjectConfigPath = NewProjectPath / TEXT("Config");
+	const FString NewProjectFolder = FPaths::GetPath(InProjectInfo.ProjectFilename);
+	const FString NewProjectName = FPaths::GetBaseFilename(InProjectInfo.ProjectFilename);
+
+	FString ProjectConfigPath = NewProjectFolder / TEXT("Config");
 
 	// DefaultEngine.ini
 	{
@@ -1407,8 +1436,8 @@ bool GameProjectUtils::GenerateConfigFiles(const FString& NewProjectPath, const
 		FileContents += TEXT("[URL]") LINE_TERMINATOR;
 		FileContents += FString::Printf(TEXT("GameName=%s") LINE_TERMINATOR, *NewProjectName);
 		FileContents += LINE_TERMINATOR;
-
-		if (bCopyStarterContent)
+		
+		if (InProjectInfo.bCopyStarterContent)
 		{
 			// for generated/blank projects with starter content, set startup map to be the starter content map
 			// otherwise, we leave it to be what the template wants.
@@ -1430,7 +1459,7 @@ bool GameProjectUtils::GenerateConfigFiles(const FString& NewProjectPath, const
 				FileContents += TEXT("[/Script/EngineSettings.GameMapsSettings]") LINE_TERMINATOR;
 				FileContents += FString::Printf(TEXT("EditorStartupMap=%s") LINE_TERMINATOR, *MapPackagePath);
 				FileContents += FString::Printf(TEXT("GameDefaultMap=%s") LINE_TERMINATOR, *MapPackagePath);
-				if (bShouldGenerateCode)
+				if (InProjectInfo.bShouldGenerateCode)
 				{
 					FileContents += FString::Printf(TEXT("GlobalDefaultGameMode=\"/Script/%s.%sGameMode\"") LINE_TERMINATOR, *NewProjectName, *NewProjectName);
 				}
@@ -1447,24 +1476,6 @@ bool GameProjectUtils::GenerateConfigFiles(const FString& NewProjectPath, const
 		}
 	}
 
-	// DefaultGame.ini
-	{
-		const FString DefaultGameIniFilename = ProjectConfigPath / TEXT("DefaultGame.ini");
-		FString FileContents;
-		FileContents += TEXT("[/Script/EngineSettings.GeneralProjectSettings]") LINE_TERMINATOR;
-		FileContents += FString::Printf( TEXT("ProjectID=%s") LINE_TERMINATOR, *FGuid::NewGuid().ToString() );
-		FileContents += LINE_TERMINATOR;
-
-		if ( WriteOutputFile(DefaultGameIniFilename, FileContents, OutFailReason) )
-		{
-			OutCreatedFiles.Add(DefaultGameIniFilename);
-		}
-		else
-		{
-			return false;
-		}
-	}
-
 	// DefaultEditor.ini
 	{
 		const FString DefaultEditorIniFilename = ProjectConfigPath / TEXT("DefaultEditor.ini");
@@ -1474,6 +1485,8 @@ bool GameProjectUtils::GenerateConfigFiles(const FString& NewProjectPath, const
 		FileContents += TEXT("bReplaceBlueprintWithClass=true") LINE_TERMINATOR;
 		FileContents += TEXT("bDontLoadBlueprintOutsideEditor=true") LINE_TERMINATOR;
 		FileContents += TEXT("bBlueprintIsNotBlueprintType=true") LINE_TERMINATOR;
+		FileContents += LINE_TERMINATOR;
+		FileContents += GetHardwareConfigString(InProjectInfo);
 
 		if (WriteOutputFile(DefaultEditorIniFilename, FileContents, OutFailReason))
 		{
diff --git a/Engine/Source/Editor/GameProjectGeneration/Private/GameProjectUtils.h b/Engine/Source/Editor/GameProjectGeneration/Private/GameProjectUtils.h
--- a/Engine/Source/Editor/GameProjectGeneration/Private/GameProjectUtils.h
+++ b/Engine/Source/Editor/GameProjectGeneration/Private/GameProjectUtils.h
@@ -2,6 +2,29 @@
 
 #pragma once
 
+#include "HardwareTargetingModule.h"
+
+struct FProjectInformation
+{
+	FProjectInformation(FString InProjectFilename, bool bInGenerateCode, bool bInCopyStarterContent, FString InTemplateFile = FString())
+		: ProjectFilename(MoveTemp(InProjectFilename))
+		, TemplateFile(MoveTemp(InTemplateFile))
+		, bShouldGenerateCode(bInGenerateCode)
+		, bCopyStarterContent(bInCopyStarterContent)
+		, TargetedHardware(EHardwareClass::Desktop)
+		, DefaultGraphicsPerformance(EGraphicsPreset::Maximum)
+	{
+	}
+	FString ProjectFilename;
+	FString TemplateFile;
+
+	bool bShouldGenerateCode;
+	bool bCopyStarterContent;
+
+	EHardwareClass::Type TargetedHardware;
+	EGraphicsPreset::Type DefaultGraphicsPerformance;
+};
+
 class GameProjectUtils
 {
 public:
@@ -133,7 +156,7 @@ public:
 	static bool OpenCodeIDE(const FString& ProjectFile, FText& OutFailReason);
 
 	/** Creates the specified project file and all required folders. If TemplateFile is non-empty, it will be used as the template for creation. On failure, OutFailReason will be populated. */
-	static bool CreateProject(const FString& NewProjectFile, const FString& TemplateFile, bool bShouldGenerateCode, bool bCopyStarterContent, FText& OutFailReason);
+	static bool CreateProject(const FProjectInformation& InProjectInfo, FText& OutFailReason);
 
 	/** Prompts the user to update his project file, if necessary. */
 	static void CheckForOutOfDateGameProjectFile();
@@ -238,11 +261,14 @@ public:
 	static void ClearSupportedTargetPlatforms();
 
 private:
+
+	static FString GetHardwareConfigString(const FProjectInformation& InProjectInfo);
+
 	/** Generates a new project without using a template project */
-	static bool GenerateProjectFromScratch(const FString& NewProjectFile, bool bShouldGenerateCode, bool bCopyStarterContent, FText& OutFailReason);
+	static bool GenerateProjectFromScratch(const FProjectInformation& InProjectInfo, FText& OutFailReason);
 
 	/** Generates a new project using a template project */
-	static bool CreateProjectFromTemplate(const FString& NewProjectFile, const FString& TemplateFile, bool bShouldGenerateCode, bool bCopyStarterContent, FText& OutFailReason);
+	static bool CreateProjectFromTemplate(const FProjectInformation& InProjectInfo, FText& OutFailReason);
 
 	/** Copies starter content into the specified project folder. */
 	static bool CopyStarterContent(const FString& DestProjectFolder, FText& OutFailReason);
@@ -278,7 +304,7 @@ private:
 	static void DeleteGeneratedBuildFiles(const FString& NewProjectFolder);
 
 	/** Creates ini files for a new project. On failure, OutFailReason will be populated. */
-	static bool GenerateConfigFiles(const FString& NewProjectPath, const FString& NewProjectName, bool bShouldGenerateCode, bool bCopyStarterContent, TArray<FString>& OutCreatedFiles, FText& OutFailReason);
+	static bool GenerateConfigFiles(const FProjectInformation& InProjectInfo, TArray<FString>& OutCreatedFiles, FText& OutFailReason);
 
 	/** Creates the basic source code for a new project. On failure, OutFailReason will be populated. */
 	static bool GenerateBasicSourceCode(const FString& NewProjectSourcePath, const FString& NewProjectName, TArray<FString>& OutGeneratedStartupModuleNames, TArray<FString>& OutCreatedFiles, FText& OutFailReason);
diff --git a/Engine/Source/Editor/GameProjectGeneration/Private/SGameProjectDialog.cpp b/Engine/Source/Editor/GameProjectGeneration/Private/SGameProjectDialog.cpp
--- a/Engine/Source/Editor/GameProjectGeneration/Private/SGameProjectDialog.cpp
+++ b/Engine/Source/Editor/GameProjectGeneration/Private/SGameProjectDialog.cpp
@@ -20,173 +20,156 @@ void SGameProjectDialog::Construct( const FArguments& InArgs )
 {
 	bool bAtLeastOneVisibleRecentProject = false;
 
-	ProjectsTabVisibility = InArgs._AllowProjectOpening ? EVisibility::Visible : EVisibility::Collapsed;
-	NewProjectTabVisibility = InArgs._AllowProjectCreate ? EVisibility::Visible : EVisibility::Collapsed;
-
-	ChildSlot
-	[
-		// Drop shadow border
-		SNew(SBorder)
-			.Padding(10.0f )
-			.BorderImage(FEditorStyle::GetBrush("ContentBrowser.ThumbnailShadow"))
+	if (InArgs._AllowProjectCreate)
+	{
+		NewProjectWizard = SNew(SNewProjectWizard);
+	}
+
+	if (InArgs._AllowProjectOpening)
+	{
+		ProjectBrowser = SNew(SProjectBrowser);
+	}
+
+	TSharedPtr<SWidget> Content;
+
+	if (InArgs._AllowProjectCreate && InArgs._AllowProjectOpening)
+	{
+		// Create the Open Project tab button
+		TSharedRef<SButton> ProjectsTabButton = SNew(SButton)
+			.ForegroundColor(FCoreStyle::Get().GetSlateColor("Foreground"))
+			.ButtonStyle(FEditorStyle::Get(), TEXT("NoBorder"))
+			.OnClicked(this, &SGameProjectDialog::HandleProjectsTabButtonClicked)
+			.ContentPadding(FMargin(40, 5))
+			.Text(LOCTEXT("ProjectsTabTitle", "Projects"))
+			.TextStyle(FEditorStyle::Get(), TEXT("ProjectBrowser.Tab.Text"));
+
+		// Create the New Project tab button
+		TSharedRef<SButton> NewProjectTabButton = SNew(SButton)
+			.ForegroundColor(FCoreStyle::Get().GetSlateColor("Foreground"))
+			.ButtonStyle(FEditorStyle::Get(), "NoBorder")
+			.OnClicked(this, &SGameProjectDialog::HandleNewProjectTabButtonClicked)
+			.ContentPadding(FMargin(20, 5))
+			.TextStyle(FEditorStyle::Get(), "ProjectBrowser.Tab.Text")
+			.Text(LOCTEXT("NewProjectTabTitle", "New Project"))
+			.ToolTip(IDocumentation::Get()->CreateToolTip(LOCTEXT("NewProjectTabTitle", "New Project"), nullptr, "Shared/LevelEditor", "NewProjectTab"));
+
+		// Allow creation and opening, so we need tabs here
+		ChildSlot
+		[
+		
+			SNew(SBorder)
+			.ColorAndOpacity(this, &SGameProjectDialog::HandleCustomContentColorAndOpacity)
+			.BorderImage(FEditorStyle::GetBrush("Docking.Tab.ContentAreaBrush"))
+			.Padding(0)
 			[
-				SNew(SBorder)
-					.BorderImage(FEditorStyle::GetBrush( "Docking.Tab.ContentAreaBrush"))
-					.Padding(10.0f)
+				SNew(SVerticalBox)
+
+				+ SVerticalBox::Slot()
+				.AutoHeight()
+				.Padding(FMargin(6.f, 0, 0, 0))
+				[
+					SNew(SHorizontalBox)
+
+					// Open Project Tab
+					+ SHorizontalBox::Slot()
+					.AutoWidth()
 					[
-						SNew(SVerticalBox)
+						SNew( SBorder )
+						.BorderImage(this, &SGameProjectDialog::OnGetTabBorderImage, ProjectsTab)
+						.Padding(0)
+						[
+							SNew(SOverlay)
+
+							+ SOverlay::Slot()
+							.VAlign(VAlign_Top)
+							[
+								SNew(SBox)
+								.HeightOverride(2.0f)
+								[
+									SNew(SImage)
+									.Image(this, &SGameProjectDialog::OnGetTabHeaderImage, ProjectsTab, ProjectsTabButton)
+									.Visibility(EVisibility::HitTestInvisible)
+								]
+							]
+
+							+ SOverlay::Slot()
+							[
+								ProjectsTabButton
+							]
+						]
+					]
 
-						+ SVerticalBox::Slot()
-							.AutoHeight()
+					+ SHorizontalBox::Slot()
+					.Padding(FMargin(6.f, 0, 0, 0))
+					.AutoWidth()
+					[
+						SNew(SBorder)
+						.BorderImage(this, &SGameProjectDialog::OnGetTabBorderImage, NewProjectTab)
+						.Padding(0)
+						[
+							SNew(SOverlay)
+
+							+ SOverlay::Slot()
+							.VAlign(VAlign_Top)
 							[
-								SNew(SHorizontalBox)
-
-								+ SHorizontalBox::Slot()
-									.AutoWidth()
-									[
-										SNew( SBorder )
-											.Visibility(ProjectsTabVisibility)
-											.BorderImage(this, &SGameProjectDialog::HandleProjectsTabHeaderBorderImage)
-											.Padding(0)
-											[
-												SNew(SOverlay)
-
-												+ SOverlay::Slot()
-													.VAlign(VAlign_Top)
-													[
-														SNew(SBox)
-															.HeightOverride(2.0f)
-															[
-																SNew(SImage)
-																	.Image(this, &SGameProjectDialog::HandleProjectsTabHeaderImage)
-																	.Visibility(EVisibility::HitTestInvisible)
-															]
-													]
-
-												+ SOverlay::Slot()
-													[
-														SAssignNew(ProjectsTabButton, SButton)
-															.ForegroundColor(FCoreStyle::Get().GetSlateColor("Foreground"))
-															.ButtonStyle(FEditorStyle::Get(), TEXT("NoBorder"))
-															.OnClicked(this, &SGameProjectDialog::HandleProjectsTabButtonClicked)
-															.ContentPadding(FMargin(40, 5))
-															.Text(LOCTEXT("ProjectsTabTitle", "Projects"))
-															.TextStyle(FEditorStyle::Get(), TEXT("ProjectBrowser.Tab.Text"))
-													]
-											]
-									]
-
-								+ SHorizontalBox::Slot()
-									.Padding(InArgs._AllowProjectOpening ? 10 : 0, 0)
-									.AutoWidth()
-									[
-										SNew(SBorder)
-											.Visibility(NewProjectTabVisibility)
-											.BorderImage(this, &SGameProjectDialog::HandleNewProjectTabHeaderBorderImage)
-											.Padding(0)
-											[
-												SNew(SOverlay)
-
-												+ SOverlay::Slot()
-													.VAlign(VAlign_Top)
-													[
-														SNew(SBox)
-														.HeightOverride(2.0f)
-														[
-															SNew(SImage)
-																.Image(this, &SGameProjectDialog::HandleNewProjectTabHeaderImage)
-																.Visibility(EVisibility::HitTestInvisible)
-														]
-													]
-
-												+ SOverlay::Slot()
-													[
-														SAssignNew(NewProjectTabButton, SButton)
-															.ForegroundColor(FCoreStyle::Get().GetSlateColor("Foreground"))
-															.ButtonStyle(FEditorStyle::Get(), "NoBorder")
-															.OnClicked(this, &SGameProjectDialog::HandleNewProjectTabButtonClicked)
-															.ContentPadding(FMargin(20, 5))
-															.TextStyle(FEditorStyle::Get(), "ProjectBrowser.Tab.Text")
-															.Text(LOCTEXT("NewProjectTabTitle", "New Project"))
-															.ToolTip(IDocumentation::Get()->CreateToolTip(LOCTEXT("NewProjectTabTitle", "New Project"), nullptr, "Shared/LevelEditor", "NewProjectTab"))
-													]
-											]
-									]
-
-								+ SHorizontalBox::Slot()
-									.FillWidth(1.0f)
-									.HAlign(HAlign_Right)
-									.Padding(0,5)
-									[
-										SNew(SBorder)
-											.BorderImage(FEditorStyle::Get().GetBrush( "ToolPanel.GroupBorder" ))
-											.Padding(3)
-											[
-												SAssignNew(MarketplaceButton, SButton)
-													.Visibility( FDesktopPlatformModule::Get()->CanOpenLauncher(true) ? EVisibility::Collapsed : EVisibility::Visible )
-													.ForegroundColor(this, &SGameProjectDialog::HandleMarketplaceTabButtonForegroundColor)
-													.ButtonStyle(FEditorStyle::Get(), "ToggleButton")
-													.OnClicked(this, &SGameProjectDialog::HandleMarketplaceTabButtonClicked)
-													.ContentPadding(FMargin(20, 0))
-													.ToolTipText(LOCTEXT("MarketplaceToolTip", "Check out the Marketplace to find new projects!"))
-													.HAlign(HAlign_Center)
-													.VAlign(VAlign_Center)
-													.Content()
-													[
-														SNew(SHorizontalBox)
-
-														+ SHorizontalBox::Slot()
-															.Padding(2.0f)
-															.AutoWidth()
-															[
-																SNew(SImage)
-																	.Image(FEditorStyle::GetBrush("LevelEditor.OpenMarketplace")) 
-															]
-
-														+ SHorizontalBox::Slot()
-															.Padding( 2.0f )
-															[
-																SNew(STextBlock)
-																	.TextStyle(FEditorStyle::Get(), "ProjectBrowser.Tab.Text")
-																	.Text(LOCTEXT("Marketplace", "Marketplace"))
-															]
-													]
-											]
-									]
+								SNew(SBox)
+								.HeightOverride(2.0f)
+								[
+									SNew(SImage)
+									.Image(this, &SGameProjectDialog::OnGetTabHeaderImage, NewProjectTab, NewProjectTabButton)
+									.Visibility(EVisibility::HitTestInvisible)
+								]
 							]
 
-						+ SVerticalBox::Slot()
+							+ SOverlay::Slot()
 							[
-								// custom content area
-								SNew(SBorder)
-									.ColorAndOpacity(this, &SGameProjectDialog::HandleCustomContentColorAndOpacity)
-									.BorderImage(FEditorStyle::GetBrush("ProjectBrowser.Background"))
-									.Padding(10.0f)
-									[
-										SAssignNew(ContentAreaSwitcher, SWidgetSwitcher)
-											.WidgetIndex(InArgs._AllowProjectOpening ? 0 : 1)
-
-										+ SWidgetSwitcher::Slot()
-											[
-												// project browser
-												SAssignNew(ProjectBrowser, SProjectBrowser)
-													.AllowProjectCreate(InArgs._AllowProjectCreate)
-													.OnNewProjectScreenRequested(this, &SGameProjectDialog::ShowNewProjectTab)
-											]
-
-										+ SWidgetSwitcher::Slot()
-											[
-												// new project wizard
-												SAssignNew(NewProjectWizard, SNewProjectWizard)
-											]
-									]
+								NewProjectTabButton
 							]
+						]
+					]
+				]
+
+				+ SVerticalBox::Slot()
+				[
+					SAssignNew(ContentAreaSwitcher, SWidgetSwitcher)
+					.WidgetIndex(0)
+
+					+ SWidgetSwitcher::Slot()
+					[
+						ProjectBrowser.ToSharedRef()
 					]
+
+					+ SWidgetSwitcher::Slot()
+					[
+						NewProjectWizard.ToSharedRef()
+					]
+				]
 			]
-	];
+		];
+	}
+	else
+	{
+		TSharedPtr<SWidget> BorderContent;
+		if (NewProjectWizard.IsValid())
+		{
+			BorderContent = NewProjectWizard;
+		}
+		else
+		{
+			BorderContent = ProjectBrowser;
+		}
+
+		ChildSlot
+		[
+			BorderContent.ToSharedRef()
+		];
+	}
 
-	ActiveTab = InArgs._AllowProjectOpening ? SGameProjectDialog::ProjectsTab : SGameProjectDialog::NewProjectTab;
-	ActiveTab = !ProjectBrowser->HasProjects() && InArgs._AllowProjectCreate ? SGameProjectDialog::NewProjectTab : ActiveTab;
+	ActiveTab = InArgs._AllowProjectOpening ? ProjectsTab : NewProjectTab;
+	if (ProjectBrowser.IsValid())
+	{
+		ActiveTab = !ProjectBrowser->HasProjects() && InArgs._AllowProjectCreate ? NewProjectTab : ActiveTab;
+	}
 
 	if (ActiveTab == ProjectsTab)
 	{
@@ -278,60 +261,6 @@ FLinearColor SGameProjectDialog::HandleCustomContentColorAndOpacity( ) const
 }
 
 
-FReply SGameProjectDialog::HandleMarketplaceTabButtonClicked( )
-{
-	IDesktopPlatform* DesktopPlatform = FDesktopPlatformModule::Get();
-
-	if (DesktopPlatform != nullptr)
-	{
-		TArray<FAnalyticsEventAttribute> EventAttributes;
-
-		if (DesktopPlatform->OpenLauncher(false, TEXT("-OpenMarket")))
-		{
-			EventAttributes.Add(FAnalyticsEventAttribute(TEXT("OpenSucceeded"), TEXT("TRUE")));
-		}
-		else
-		{
-			EventAttributes.Add(FAnalyticsEventAttribute(TEXT("OpenSucceeded"), TEXT("FALSE")));
-
-			if (EAppReturnType::Yes == FMessageDialog::Open(EAppMsgType::YesNo, LOCTEXT("InstallMarketplacePrompt", "The Marketplace requires the Unreal Engine Launcher, which does not seem to be installed on your computer. Would you like to install it now?")))
-			{
-				if (!DesktopPlatform->OpenLauncher(true, TEXT("-OpenMarket")))
-				{
-					EventAttributes.Add(FAnalyticsEventAttribute(TEXT("InstallSucceeded"), TEXT("FALSE")));
-					FMessageDialog::Open(EAppMsgType::Ok, FText::FromString(TEXT("Sorry, there was a problem installing the Launcher.\nPlease try to install it manually!")));
-				}
-				else
-				{
-					EventAttributes.Add(FAnalyticsEventAttribute(TEXT("InstallSucceeded"), TEXT("TRUE")));
-				}
-			}
-		}
-
-		EventAttributes.Add(FAnalyticsEventAttribute(TEXT("Source"), TEXT("ProjectBrowser")));
-		if( FEngineAnalytics::IsAvailable() )
-		{
-			FEngineAnalytics::GetProvider().RecordEvent(TEXT("Editor.Usage.OpenMarketplace"), EventAttributes);
-		}
-	}
-
-	return FReply::Handled();
-}
-
-
-FSlateColor SGameProjectDialog::HandleMarketplaceTabButtonForegroundColor( ) const
-{
-	//if (MarketplaceButton->IsHovered())
-	//{
-	//	return FLinearColor::Black;
-	//}
-	//else
-	//{
-		return FSlateColor::UseForeground();
-	//}
-}
-
-
 FReply SGameProjectDialog::HandleNewProjectTabButtonClicked( )
 {
 	ShowNewProjectTab();
@@ -340,25 +269,14 @@ FReply SGameProjectDialog::HandleNewProjectTabButtonClicked( )
 }
 
 
-const FSlateBrush* SGameProjectDialog::HandleNewProjectTabHeaderBorderImage( ) const
+const FSlateBrush* SGameProjectDialog::OnGetTabHeaderImage( ETab InTab, TSharedRef<SButton> TabButton ) const
 {
-	if (ActiveTab == NewProjectTab)
-	{
-		return FEditorStyle::GetBrush("ProjectBrowser.Tab.ActiveBackground");
-	}
-	
-	return FEditorStyle::GetBrush("ProjectBrowser.Tab.Background");
-}
-
-
-const FSlateBrush* SGameProjectDialog::HandleNewProjectTabHeaderImage( ) const
-{
-	if (NewProjectTabButton->IsPressed())
+	if (TabButton->IsPressed())
 	{
 		return FEditorStyle::GetBrush("ProjectBrowser.Tab.PressedHighlight");
 	}
 
-	if ((ActiveTab == NewProjectTab) || NewProjectTabButton->IsHovered())
+	if ((ActiveTab == InTab) || TabButton->IsHovered())
 	{
 		return FEditorStyle::GetBrush("ProjectBrowser.Tab.ActiveHighlight");
 	}
@@ -374,10 +292,9 @@ FReply SGameProjectDialog::HandleProjectsTabButtonClicked( )
 	return FReply::Handled();
 }
 
-
-const FSlateBrush* SGameProjectDialog::HandleProjectsTabHeaderBorderImage( ) const
+const FSlateBrush* SGameProjectDialog::OnGetTabBorderImage( ETab InTab ) const
 {
-	if (ActiveTab == ProjectsTab)
+	if (ActiveTab == InTab)
 	{
 		return FEditorStyle::GetBrush("ProjectBrowser.Tab.ActiveBackground");
 	}
@@ -386,20 +303,4 @@ const FSlateBrush* SGameProjectDialog::HandleProjectsTabHeaderBorderImage( ) con
 }
 
 
-const FSlateBrush* SGameProjectDialog::HandleProjectsTabHeaderImage( ) const
-{
-	if (ProjectsTabButton->IsPressed())
-	{
-		return FEditorStyle::GetBrush("ProjectBrowser.Tab.PressedHighlight");
-	}
-
-	if ((ActiveTab == ProjectsTab) || ProjectsTabButton->IsHovered())
-	{
-		return FEditorStyle::GetBrush("ProjectBrowser.Tab.ActiveHighlight");
-	}
-	
-	return FEditorStyle::GetBrush("ProjectBrowser.Tab.Highlight");
-}
-
-
 #undef LOCTEXT_NAMESPACE
diff --git a/Engine/Source/Editor/GameProjectGeneration/Private/SGameProjectDialog.h b/Engine/Source/Editor/GameProjectGeneration/Private/SGameProjectDialog.h
--- a/Engine/Source/Editor/GameProjectGeneration/Private/SGameProjectDialog.h
+++ b/Engine/Source/Editor/GameProjectGeneration/Private/SGameProjectDialog.h
@@ -66,29 +66,17 @@ private:
 	// Callback for getting the color of the custom content area.
 	FLinearColor HandleCustomContentColorAndOpacity() const;
 
-	// Callback for clicking the 'Marketplace' button.
-	FReply HandleMarketplaceTabButtonClicked( );
-
-	// Callback for getting the foreground color of the 'Marketplace' button.
-	FSlateColor HandleMarketplaceTabButtonForegroundColor( ) const;
-
 	// Callback for clicking the 'New Project' button.
 	FReply HandleNewProjectTabButtonClicked( );
 
-	// Callback for getting the border image of the 'New Project' button.
-	const FSlateBrush* HandleNewProjectTabHeaderBorderImage( ) const;
-
-	// Callback for getting the button image of the 'New Project' button.
-	const FSlateBrush* HandleNewProjectTabHeaderImage( ) const;
-
 	// Callback for clicking the 'Projects' button.
 	FReply HandleProjectsTabButtonClicked( );
 
-	// Callback for getting the border image of the 'Projects' button.
-	const FSlateBrush* HandleProjectsTabHeaderBorderImage( ) const;
+	// Callback for getting the border image of the specified tab.
+	const FSlateBrush* OnGetTabBorderImage( ETab InTab ) const;
 
-	// Callback for getting the button image of the 'Projects' button.
-	const FSlateBrush* HandleProjectsTabHeaderImage( ) const;
+	// Callback for getting the header stripe image for the specified tab.
+	const FSlateBrush* OnGetTabHeaderImage( ETab InTab, TSharedRef<SButton> TabButton ) const;
 
 private:
 
@@ -102,12 +90,6 @@ private:
 	TSharedPtr<SWidgetSwitcher> ContentAreaSwitcher;
 	TSharedPtr<SProjectBrowser> ProjectBrowser;
 	TSharedPtr<SNewProjectWizard> NewProjectWizard;
-	TSharedPtr<SButton> ProjectsTabButton;
-	TSharedPtr<SButton> NewProjectTabButton;
-	TSharedPtr<SButton> MarketplaceButton;
 
 	ETab ActiveTab;
-
-	EVisibility ProjectsTabVisibility;
-	EVisibility NewProjectTabVisibility;
 };
diff --git a/Engine/Source/Editor/GameProjectGeneration/Private/SNewProjectWizard.cpp b/Engine/Source/Editor/GameProjectGeneration/Private/SNewProjectWizard.cpp
--- a/Engine/Source/Editor/GameProjectGeneration/Private/SNewProjectWizard.cpp
+++ b/Engine/Source/Editor/GameProjectGeneration/Private/SNewProjectWizard.cpp
@@ -9,7 +9,7 @@
 #include "TemplateCategory.h"
 #include "GameProjectGenerationModule.h"
 #include "SWizard.h"
-
+#include "HardwareTargetingModule.h"
 
 #define LOCTEXT_NAMESPACE "NewProjectWizard"
 
@@ -240,9 +240,9 @@ public:
 				+ SVerticalBox::Slot()
 				.AutoHeight()
 				.HAlign(HAlign_Center)
+				.Padding(ThumbnailPadding)
 				[
 					SNew(SBox)
-					.Padding(ThumbnailPadding)
 					.WidthOverride(ThumbnailSize)
 					.HeightOverride(ThumbnailSize)
 					[
@@ -255,9 +255,10 @@ public:
 				+ SVerticalBox::Slot()
 				.HAlign(HAlign_Center)
 				.VAlign(VAlign_Top)
+				.Padding(FMargin(ThumbnailPadding, 0))
 				[
 					SNew(SRichTextBlock)
-					.WrapTextAt(ItemWidth - ThumbnailPadding*2)
+					.AutoWrapText(true)
 					.Justification(ETextJustify::Center)
 					//.HighlightText(this, &SNewProjectWizard::GetItemHighlightText)
 					.Text(InArgs._Item->Name)
@@ -277,7 +278,7 @@ private:
 		{
 			return ItemPtr->Thumbnail.Get();
 		}
-		return FEditorStyle::GetBrush("GameProjectDialog.DefaultGameThumbnail");
+		return FEditorStyle::GetBrush("GameProjectDialog.DefaultGameThumbnail.Small");
 	}
 
 };
@@ -290,7 +291,12 @@ void SNewProjectWizard::Construct( const FArguments& InArgs )
 	bLastGlobalValidityCheckSuccessful = true;
 	bLastNameAndLocationValidityCheckSuccessful = true;
 	bPreventPeriodicValidityChecksUntilNextChange = false;
-	 
+	
+	IHardwareTargetingModule& HardwareTargeting = IHardwareTargetingModule::Get();
+
+	SelectedHardwareClassTarget = EHardwareClass::Desktop;
+	SelectedGraphicsPreset = EGraphicsPreset::Maximum;
+
 	// Find all template projects
 	FindTemplateProjects();
 	SetDefaultProjectLocation();
@@ -305,202 +311,281 @@ void SNewProjectWizard::Construct( const FArguments& InArgs )
 	.OnMouseButtonDoubleClick(this, &SNewProjectWizard::HandleTemplateListViewDoubleClick)
 	.OnSelectionChanged(this, &SNewProjectWizard::HandleTemplateListViewSelectionChanged);
 
-	const FMargin StandardVerticalPad(0.f, 0.f, 0.f, 10.f);
 	const EVisibility StarterContentVisiblity = GameProjectUtils::IsStarterContentAvailableForNewProjects() ? EVisibility::Visible : EVisibility::Collapsed;
 
+	TSharedRef<SSeparator> Separator = SNew(SSeparator).Orientation(EOrientation::Orient_Vertical);
+	Separator->SetBorderBackgroundColor(FLinearColor::White.CopyWithNewOpacity(0.25f));
+
+	const float UniformPadding = 16.f;
 	ChildSlot
 	[
-		SNew(SOverlay)
-
-		// Wizard
-		+SOverlay::Slot()
+		SNew(SBorder)
+		.BorderImage( FEditorStyle::GetBrush("ToolPanel.GroupBorder") )
 		[
-			SAssignNew(MainWizard, SWizard)
-			.ShowPageList(false)
-			.ShowCancelButton(false)
-			.CanFinish(this, &SNewProjectWizard::HandleCreateProjectWizardCanFinish)
-			.FinishButtonText(LOCTEXT("FinishButtonText", "Create Project"))
-			.FinishButtonToolTip(LOCTEXT("FinishButtonToolTip", "Creates your new project in the specified location with the specified template and name.") )
-			.OnFinished(this, &SNewProjectWizard::HandleCreateProjectWizardFinished)
-			.OnFirstPageBackClicked(InArgs._OnBackRequested)
-
-			// Choose Template
-			+SWizard::Page()
-			.OnEnter(this, &SNewProjectWizard::OnPageVisited, TemplatePageName)
+			SNew(SOverlay)
+
+			// Wizard
+			+SOverlay::Slot()
+			.Padding(UniformPadding / 2)
 			[
-				SNew(SBorder)
-				.BorderImage(FEditorStyle::GetBrush("NoBorder"))
-				.Padding(10.f)
+				SAssignNew(MainWizard, SWizard)
+				.ShowPageList(false)
+				.ShowCancelButton(false)
+				.CanFinish(this, &SNewProjectWizard::HandleCreateProjectWizardCanFinish)
+				.FinishButtonText(LOCTEXT("FinishButtonText", "Create Project"))
+				.FinishButtonToolTip(LOCTEXT("FinishButtonToolTip", "Creates your new project in the specified location with the specified template and name.") )
+				.OnFinished(this, &SNewProjectWizard::HandleCreateProjectWizardFinished)
+				.OnFirstPageBackClicked(InArgs._OnBackRequested)
+
+				// Choose Template
+				+SWizard::Page()
+				.OnEnter(this, &SNewProjectWizard::OnPageVisited, TemplatePageName)
 				[
-					SNew(SHorizontalBox)
-
-					// Template category tabs
-					+ SHorizontalBox::Slot()
+					SNew(SBorder)
+					.BorderImage(FEditorStyle::GetBrush("NoBorder"))
+					.Padding(FMargin(UniformPadding / 2, UniformPadding / 2, UniformPadding / 2, 0))
 					[
 						SNew(SVerticalBox)
 
-						+SVerticalBox::Slot()
-						.Padding(FMargin(5.f, 0.f))
+						+ SVerticalBox::Slot()
+						.Padding(FMargin(0, 0, 0, 15))
 						.AutoHeight()
 						[
-							BuildCategoryTabs()
+							SNew(STextBlock)
+							.Text(LOCTEXT("ProjectTemplateDescription", "Choose a template to use as a starting point for your new project:"))
 						]
 
-						// Templates list
 						+ SVerticalBox::Slot()
-						.FillHeight(1.0f)
-						// Manual padding offset to align to the bottom of the file path widget
-						.Padding(FMargin(0,0,0,13.f))
 						[
-							SNew(SBorder)
-							.Padding(10)
-							.BorderImage(FEditorStyle::GetBrush("GameProjectDialog.TabBackground"))
+							// Template category tabs
+							SNew(SVerticalBox)
+
+							+SVerticalBox::Slot()
+							.Padding(FMargin(8.f, 0.f))
+							.AutoHeight()
 							[
-								SNew(SScrollBorder, TemplateListView.ToSharedRef())
-								[
-									TemplateListView.ToSharedRef()
-								]
+								BuildCategoryTabs()
 							]
-						]
-					]
 
-					// Selected template details
-					+ SHorizontalBox::Slot()
-					.AutoWidth()
-					// top-offset fudge to account for size of tabs
-					.Padding(FMargin(20.0f, 42.f, 0.f, 0.f))
-					[
-						// This pane is always 512px wide - the width of the template preview
-						SNew(SBox)
-						.WidthOverride(512)
-						[
-							SNew(SVerticalBox)
-
-							// Preview image
+							// Templates list
 							+ SVerticalBox::Slot()
-							.AutoHeight()
-							.HAlign(HAlign_Center)
-							.Padding(StandardVerticalPad)
+							.FillHeight(1.0f)
 							[
-								SNew(SOverlay)
-							
-								+ SOverlay::Slot()
+								SNew(SBorder)
+								.Padding(UniformPadding)
+								.BorderImage(FEditorStyle::GetBrush("GameProjectDialog.TabBackground"))
 								[
-									SNew(SBorder)
-									.Padding(FMargin(0, 0, 0, 5.f))
-									.BorderImage(FEditorStyle::GetBrush("ContentBrowser.ThumbnailShadow"))
+									SNew(SHorizontalBox)
+
+									+ SHorizontalBox::Slot()
 									[
-										SNew(SImage)
-										.Image(this, &SNewProjectWizard::GetSelectedTemplatePreviewImage)
+										SNew(SScrollBorder, TemplateListView.ToSharedRef())
+										[
+											TemplateListView.ToSharedRef()
+										]
+									]
+										
+									+ SHorizontalBox::Slot()
+									.Padding(UniformPadding, 0)
+									.AutoWidth()
+									[
+										Separator
 									]
-								]
 
-								+ SOverlay::Slot()
-								.HAlign(HAlign_Right)
-								.VAlign(VAlign_Top)
-								.Padding(10)
-								[
-									SNew(SBox)
-									.WidthOverride(48)
-									.HeightOverride(48)
+									// Selected template details
+									+ SHorizontalBox::Slot()
 									[
-										SNew(SImage)
-										.Image(this, &SNewProjectWizard::GetSelectedTemplateTypeImage)
+										SNew(SScrollBox)
+										+ SScrollBox::Slot()
+										.Padding(UniformPadding, 0)
+										[
+											SNew(SVerticalBox)
+
+											// Preview image
+											+ SVerticalBox::Slot()
+											.AutoHeight()
+											.HAlign(HAlign_Center)
+											.Padding(FMargin(0, 0, 0, 15.f))
+											[
+												SNew(SBox)
+												.Visibility(this, &SNewProjectWizard::GetSelectedTemplatePreviewVisibility)
+												.WidthOverride(400)
+												.HeightOverride(200)
+												[
+													SNew(SOverlay)
+							
+													+ SOverlay::Slot()
+													[
+														SNew(SBorder)
+														.Padding(FMargin(0, 0, 0, 5.f))
+														.BorderImage(FEditorStyle::GetBrush("ContentBrowser.ThumbnailShadow"))
+														[
+															SNew(SImage)
+															.Image(this, &SNewProjectWizard::GetSelectedTemplatePreviewImage)
+														]
+													]
+
+													+ SOverlay::Slot()
+													.HAlign(HAlign_Right)
+													.VAlign(VAlign_Top)
+													.Padding(10)
+													[
+														SNew(SBox)
+														.WidthOverride(48)
+														.HeightOverride(48)
+														[
+															SNew(SImage)
+															.Image(this, &SNewProjectWizard::GetSelectedTemplateTypeImage)
+														]
+													]
+												]
+											]
+
+											// Template Name
+											+ SVerticalBox::Slot()
+											.Padding(FMargin(0, 0, 0, 10))
+											.AutoHeight()
+											[
+												SNew(STextBlock)
+												.AutoWrapText(true)
+												.TextStyle(FEditorStyle::Get(), "GameProjectDialog.FeatureText")
+												.Text(this, &SNewProjectWizard::GetSelectedTemplateProperty<FText>, &FTemplateItem::Name)
+											]
+						
+											// Template Description
+											+ SVerticalBox::Slot()
+											[
+												SNew(STextBlock)
+												.AutoWrapText(true)
+												.Text(this, &SNewProjectWizard::GetSelectedTemplateProperty<FText>, &FTemplateItem::Description)
+											]
+										]
 									]
+
+
 								]
 							]
+						]
 
-							// Template Name
-							+ SVerticalBox::Slot()
-							.Padding(StandardVerticalPad)
-							.AutoHeight()
-							[
-								SNew(STextBlock)
-								.AutoWrapText(true)
-								.TextStyle(FEditorStyle::Get(), "GameProjectDialog.FeatureText")
-								.Text(this, &SNewProjectWizard::GetSelectedTemplateProperty<FText>, &FTemplateItem::Name)
-							]
-						
-							// Template Description
-							+ SVerticalBox::Slot()
-							.Padding(StandardVerticalPad)
+						+ SVerticalBox::Slot()
+						.AutoHeight()
+						.Padding(FMargin(0, 15, 0, 0))
+						[
+							SNew(SScrollBox)
+
+							+ SScrollBox::Slot()
 							[
-								SNew(SScrollBox)
-								+ SScrollBox::Slot()
+								
+								SNew(SVerticalBox)
+									
+								+ SVerticalBox::Slot()
+								.AutoHeight()
+								.Padding(FMargin(0, 0, 0, 25.f))
 								[
 									SNew(STextBlock)
-									.AutoWrapText(true)
-									.Text(this, &SNewProjectWizard::GetSelectedTemplateProperty<FText>, &FTemplateItem::Description)
+									.Text(LOCTEXT("ProjectSettingsDescription", "Now choose some settings for your project. Don't worry, you can still change these after the project has been created:"))
 								]
-							]
-						
-							// Include Starter Content?
-							+ SVerticalBox::Slot()
-							.AutoHeight()
-							.Padding(StandardVerticalPad)
-							[
-								SNew(SCheckBox)
-								.Visibility(StarterContentVisiblity)
-								.ToolTipText( LOCTEXT("CopyStarterContent_ToolTip", "Adds example content to your new project, including simple placeable meshes with basic materials and textures.\nYou can turn this off to start with a project that is only has the bare essentials for the selected project template."))
-								.OnCheckStateChanged(this, &SNewProjectWizard::OnCopyStarterContentChanged)
-								.IsChecked(this, &SNewProjectWizard::IsCopyStarterContentChecked)
-								.Content()
-								[
-									SNew(STextBlock)
-									.Text(LOCTEXT("CopyStarterContent", "Include starter content" ))
- 								]
-							]
 
-							// Project name error text
-							+ SVerticalBox::Slot()
-							.AutoHeight()
-							.Padding(StandardVerticalPad)
-							[
-								SNew(SBorder)
-								.BorderImage(FEditorStyle::GetBrush("GameProjectDialog.ErrorLabelBorder"))
-								.Visibility(this, &SNewProjectWizard::GetNameAndLocationErrorLabelVisibility)
-								.Padding(5.f)
-								.Content()
+								+ SVerticalBox::Slot()
+								.HAlign(HAlign_Center)
+								.AutoHeight()
 								[
-									SNew(STextBlock)
-									.Text(this, &SNewProjectWizard::GetNameAndLocationErrorLabelText)
-									.TextStyle(FEditorStyle::Get(), "GameProjectDialog.ErrorLabelFont")
+									SNew(SBox)
+									.WidthOverride(650)
+									[
+										SNew(SVerticalBox)
+
+										+ SVerticalBox::Slot()
+										.AutoHeight()
+										.HAlign(HAlign_Center)
+										.Padding(FMargin(0, 0, 0, 25.f))
+										[
+											SNew(SHorizontalBox)
+
+											+ SHorizontalBox::Slot()
+											.AutoWidth()
+											[
+												HardwareTargeting.MakeHardwareClassTargetCombo(
+													FOnHardwareClassChanged::CreateSP(this, &SNewProjectWizard::SetHardwareClassTarget),
+													TAttribute<EHardwareClass::Type>(this, &SNewProjectWizard::GetHardwareClassTarget)
+													)
+											]
+								
+											+ SHorizontalBox::Slot()
+											.AutoWidth()
+											.Padding(FMargin(30, 0))
+											[
+												HardwareTargeting.MakeGraphicsPresetTargetCombo(
+													FOnGraphicsPresetChanged::CreateSP(this, &SNewProjectWizard::SetGraphicsPreset),
+													TAttribute<EGraphicsPreset::Type>(this, &SNewProjectWizard::GetGraphicsPreset)
+													)
+											]
+								
+											+ SHorizontalBox::Slot()
+											.AutoWidth()
+											[
+												SNew(SButton)
+												.Visibility(StarterContentVisiblity)
+												.ButtonStyle(FEditorStyle::Get(), "HoverHintOnly")
+												.ForegroundColor(FSlateColor::UseForeground())
+												.ContentPadding(6.f)
+												.ToolTipText( LOCTEXT("CopyStarterContent_ToolTip", "Enable to include an additional content pack containing simple placeable meshes with basic materials and textures.\nYou can opt out of including this to create a project that only has the bare essentials for the selected project template."))
+												.OnClicked(this, &SNewProjectWizard::OnCopyStarterContentClicked)
+												[
+													SNew(SVerticalBox)
+
+													+ SVerticalBox::Slot()
+													.AutoHeight()
+													.HAlign(HAlign_Center)
+													[
+														SNew(SImage)
+														.Image(this, &SNewProjectWizard::GetStartContentIcon)
+													]
+
+													+ SVerticalBox::Slot()
+													.AutoHeight()
+													.Padding(4)
+													.HAlign(HAlign_Center)
+													[
+														SNew(STextBlock)
+														.Text(this, &SNewProjectWizard::GetStartContentText)
+													]
+												]
+											]
+										]
+
+										// File path widget
+										+ SVerticalBox::Slot()
+										.AutoHeight()
+										.HAlign(HAlign_Center)
+										[
+											SNew(SFilepath)
+											.OnBrowseForFolder(this, &SNewProjectWizard::HandleBrowseButtonClicked)
+											.LabelBackgroundBrush(FEditorStyle::GetBrush("ProjectBrowser.Background"))
+											.LabelBackgroundColor(FLinearColor::White)
+											.FolderPath(this, &SNewProjectWizard::GetCurrentProjectFilePath)
+											.Name(this, &SNewProjectWizard::GetCurrentProjectFileName)
+											.OnFolderChanged(this, &SNewProjectWizard::OnCurrentProjectFilePathChanged)
+											.OnNameChanged(this, &SNewProjectWizard::OnCurrentProjectFileNameChanged)
+										]
+									]
 								]
 							]
-
-							// File path widget
-							+ SVerticalBox::Slot()
-							.AutoHeight()
-							[
-								SNew(SFilepath)
-								.OnBrowseForFolder(this, &SNewProjectWizard::HandleBrowseButtonClicked)
-								.LabelBackgroundBrush(FEditorStyle::GetBrush("ProjectBrowser.Background"))
-								.LabelBackgroundColor(FLinearColor::White)
-								.FolderPath(this, &SNewProjectWizard::GetCurrentProjectFilePath)
-								.Name(this, &SNewProjectWizard::GetCurrentProjectFileName)
-								.OnFolderChanged(this, &SNewProjectWizard::OnCurrentProjectFilePathChanged)
-								.OnNameChanged(this, &SNewProjectWizard::OnCurrentProjectFileNameChanged)
-							]
 						]
 					]
 				]
 			]
-		]
 
-		// Global Error label
-		+SOverlay::Slot()
-		.HAlign(HAlign_Left)
-		.VAlign(VAlign_Bottom)
-		.Padding( FMargin( 0.0f, 5.0f, 0.0f, 0.0f ) )
-		[
-			// Constant height, whether the label is visible or not
-			SNew(SBox)
-			.HeightOverride(20.0f)
+			// Global Error label
+			+SOverlay::Slot()
+			.HAlign(HAlign_Left)
+			.VAlign(VAlign_Bottom)
+			.Padding( UniformPadding / 2 )
 			[
 				SNew(SBorder)
 				.Visibility(this, &SNewProjectWizard::GetGlobalErrorLabelVisibility)
 				.BorderImage(FEditorStyle::GetBrush("GameProjectDialog.ErrorLabelBorder"))
-				.Content()
+				.Padding( UniformPadding / 2 )
 				[
 					SNew(SHorizontalBox)
 										
@@ -541,6 +626,31 @@ void SNewProjectWizard::Construct( const FArguments& InArgs )
 					]
 				]
 			]
+
+			// Project filename error
+			+SOverlay::Slot()
+			.HAlign(HAlign_Left)
+			.VAlign(VAlign_Bottom)
+			.Padding( UniformPadding / 2 )
+			[
+				SNew(SBorder)
+				.BorderImage(FEditorStyle::GetBrush("GameProjectDialog.ErrorLabelBorder"))
+				.Visibility(this, &SNewProjectWizard::GetNameAndLocationErrorLabelVisibility)
+				.Padding(UniformPadding / 2)
+				[
+					SNew(SHorizontalBox)
+										
+					+SHorizontalBox::Slot()
+					.VAlign(VAlign_Center)
+					.FillWidth(1.0f)
+					[
+						SNew(STextBlock)
+						.AutoWrapText(true)
+						.Text(this, &SNewProjectWizard::GetNameAndLocationErrorLabelText)
+						.TextStyle(FEditorStyle::Get(), "GameProjectDialog.ErrorLabelFont")
+					]
+				]
+			]
 		]
 	];
 
@@ -616,13 +726,10 @@ TSharedRef<SWidget> SNewProjectWizard::BuildCategoryTabs()
 	return TabStrip;
 }
 
-void SNewProjectWizard::OnCopyStarterContentChanged(ESlateCheckBoxState::Type InNewState)
-{
-	bCopyStarterContent = InNewState == ESlateCheckBoxState::Checked;
-}
-ESlateCheckBoxState::Type SNewProjectWizard::IsCopyStarterContentChecked() const
+FReply SNewProjectWizard::OnCopyStarterContentClicked()
 {
-	return bCopyStarterContent ? ESlateCheckBoxState::Checked : ESlateCheckBoxState::Unchecked;
+	bCopyStarterContent = !bCopyStarterContent;
+	return FReply::Handled();
 }
 
 void SNewProjectWizard::Tick( const FGeometry& AllottedGeometry, const double InCurrentTime, const float InDeltaTime )
@@ -662,6 +769,12 @@ const FSlateBrush* SNewProjectWizard::GetSelectedTemplatePreviewImage() const
 	return PreviewImage.IsValid() ? PreviewImage.Get() : nullptr;
 }
 
+EVisibility SNewProjectWizard::GetSelectedTemplatePreviewVisibility() const
+{
+	auto PreviewImage = GetSelectedTemplateProperty(&FTemplateItem::PreviewImage);
+	return PreviewImage.IsValid() ? EVisibility::Visible : EVisibility::Collapsed;
+}
+
 const FSlateBrush* SNewProjectWizard::GetSelectedTemplateTypeImage() const
 {
 	TSharedPtr<FTemplateItem> SelectedItem = GetSelectedTemplateItem();
@@ -835,7 +948,8 @@ void SNewProjectWizard::OnPageVisited(FName NewPageName)
 
 EVisibility SNewProjectWizard::GetGlobalErrorLabelVisibility() const
 {
-	return GetGlobalErrorLabelText().IsEmpty() ? EVisibility::Hidden : EVisibility::Visible;
+	const bool bIsVisible = GetNameAndLocationErrorLabelText().IsEmpty() && !GetGlobalErrorLabelText().IsEmpty();
+	return bIsVisible ? EVisibility::Visible : EVisibility::Hidden;
 }
 
 
@@ -1192,10 +1306,12 @@ bool SNewProjectWizard::CreateProject( const FString& ProjectFile )
 		return false;
 	}
 
-	const bool bAddCode = SelectedTemplate->bGenerateCode;
 	FText FailReason;
 
-	if (!GameProjectUtils::CreateProject( ProjectFile, SelectedTemplate->ProjectFile, bAddCode, bCopyStarterContent, FailReason))
+	FProjectInformation ProjectInfo(ProjectFile, SelectedTemplate->bGenerateCode, bCopyStarterContent, SelectedTemplate->ProjectFile);
+	ProjectInfo.TargetedHardware = SelectedHardwareClassTarget;
+	ProjectInfo.DefaultGraphicsPerformance = SelectedGraphicsPreset;
+	if (!GameProjectUtils::CreateProject(ProjectInfo, FailReason))
 	{
 		DisplayError(FailReason);
 		return false;
@@ -1376,5 +1492,24 @@ void SNewProjectWizard::HandleCategoryChanged(ESlateCheckBoxState::Type CheckSta
 	TemplateListView->RequestListRefresh();
 }
 
+void SNewProjectWizard::SetHardwareClassTarget(EHardwareClass::Type InHardwareClass)
+{
+	SelectedHardwareClassTarget = InHardwareClass;
+}
+
+void SNewProjectWizard::SetGraphicsPreset(EGraphicsPreset::Type InGraphicsPreset)
+{
+	SelectedGraphicsPreset = InGraphicsPreset;
+}
+
+FText SNewProjectWizard::GetStartContentText() const
+{
+	return bCopyStarterContent ? LOCTEXT("IncludeStarterContent", "With Starter Content") : LOCTEXT("NoStarterContent", "No Starter Content");
+}
+
+const FSlateBrush* SNewProjectWizard::GetStartContentIcon() const
+{
+	return bCopyStarterContent ? FEditorStyle::GetBrush("GameProjectDialog.IncludeStarterContent") : FEditorStyle::GetBrush("GameProjectDialog.NoStarterContent");
+}
 
 #undef LOCTEXT_NAMESPACE
diff --git a/Engine/Source/Editor/GameProjectGeneration/Private/SNewProjectWizard.h b/Engine/Source/Editor/GameProjectGeneration/Private/SNewProjectWizard.h
--- a/Engine/Source/Editor/GameProjectGeneration/Private/SNewProjectWizard.h
+++ b/Engine/Source/Editor/GameProjectGeneration/Private/SNewProjectWizard.h
@@ -2,7 +2,6 @@
 
 #pragma once
 
-
 struct FTemplateItem;
 struct FTemplateCategory;
 
@@ -74,6 +73,9 @@ private:
 	const FSlateBrush* GetSelectedTemplatePreviewImage() const;
 	const FSlateBrush* GetSelectedTemplateTypeImage() const;
 
+	/** Get the visibility for the selected template preview image */
+	EVisibility GetSelectedTemplatePreviewVisibility() const;
+
 	/** Gets the assembled project filename with path */
 	FString GetProjectFilenameWithPath() const;
 
@@ -111,13 +113,22 @@ private:
 	FText GetNameAndLocationErrorLabelText() const;
 
 	/** Handler for when the "copy starter content" checkbox changes state */
-	void OnCopyStarterContentChanged(ESlateCheckBoxState::Type InNewState);
+	FReply OnCopyStarterContentClicked();
+
+private:
 
-	/** Returns true if the "copy starter content" checkbox should be checked. */
-	ESlateCheckBoxState::Type IsCopyStarterContentChecked() const;
+	EHardwareClass::Type SelectedHardwareClassTarget;
+	void SetHardwareClassTarget(EHardwareClass::Type InHardwareClass);
+	EHardwareClass::Type GetHardwareClassTarget() const { return SelectedHardwareClassTarget; }
+
+	EGraphicsPreset::Type SelectedGraphicsPreset;
+	void SetGraphicsPreset(EGraphicsPreset::Type InGraphicsPreset);
+	EGraphicsPreset::Type GetGraphicsPreset() const { return SelectedGraphicsPreset; }
+
+private:
 
-	/** Gets the visibility of the copy starter content option */
-	EVisibility GetCopyStarterContentCheckboxVisibility() const;
+	FText GetStartContentText() const;
+	const FSlateBrush* GetStartContentIcon() const;
 
 private:
 
diff --git a/Engine/Source/Editor/GameProjectGeneration/Private/SProjectBrowser.cpp b/Engine/Source/Editor/GameProjectGeneration/Private/SProjectBrowser.cpp
--- a/Engine/Source/Editor/GameProjectGeneration/Private/SProjectBrowser.cpp
+++ b/Engine/Source/Editor/GameProjectGeneration/Private/SProjectBrowser.cpp
@@ -97,11 +97,9 @@ SProjectBrowser::SProjectBrowser()
 BEGIN_SLATE_FUNCTION_BUILD_OPTIMIZATION
 void SProjectBrowser::Construct( const FArguments& InArgs )
 {
-	NewProjectScreenRequestedDelegate = InArgs._OnNewProjectScreenRequested;
 	bPreventSelectionChangeEvent = false;
 	ThumbnailBorderPadding = 5;
 	ThumbnailSize = 128.0f;
-	bAllowProjectCreate = InArgs._AllowProjectCreate;
 
 	// Prepare the categories box
 	CategoriesBox = SNew(SVerticalBox);
@@ -129,66 +127,121 @@ void SProjectBrowser::Construct( const FArguments& InArgs )
 
 	ChildSlot
 	[
-		SNew(SVerticalBox)
-
-		// Page description
-		+ SVerticalBox::Slot()
-		.AutoHeight()
-		.Padding(10.0f, 0.0f)
+		SNew(SBorder)
+		.BorderImage( FEditorStyle::GetBrush("ToolPanel.GroupBorder") )
+		.Padding(8.f)
 		[
-			SNew(STextBlock)
-			.Text(LOCTEXT("SelectProjectDescription", "Select an existing project from the list below."))
-		]
+			SNew(SVerticalBox)
 
-		// Categories
-		+ SVerticalBox::Slot()
-		.FillHeight(1.0f)
-		[
-			SNew(SBorder)
-			.BorderImage( FEditorStyle::GetBrush("ToolPanel.GroupBorder") )
-			.Padding(4)
+			// Categories
+			+ SVerticalBox::Slot()
+			.Padding(8.f)
+			.FillHeight(1.0f)
 			[
 				SNew(SVerticalBox)
 
 				+ SVerticalBox::Slot()
 				.AutoHeight()
 				.VAlign(VAlign_Top)
-				.Padding(FMargin(5.f, 10.f))
+				.HAlign(HAlign_Left)
+				.Padding(FMargin(0, 0, 0, 5.f))
 				[
 					SNew(SHorizontalBox)
+					
 					+SHorizontalBox::Slot()
+					.AutoWidth()
+					.Padding(FMargin(0,0,5,0))
 					[
-						SNew(SOverlay)
-
-						+SOverlay::Slot()
-						[
-							SNew(SSearchBox)
-							.HintText(LOCTEXT("FilterHint", "Filter Projects..."))
-							.OnTextChanged(this, &SProjectBrowser::OnFilterTextChanged)
-						]
-
-						+SOverlay::Slot()
+						SNew(SButton)
+						.Visibility( FDesktopPlatformModule::Get()->CanOpenLauncher(true) ? EVisibility::Collapsed : EVisibility::Visible )
+						.ButtonStyle(FEditorStyle::Get(), "ToggleButton")
+						.OnClicked(this, &SProjectBrowser::HandleMarketplaceTabButtonClicked)
+						.ForegroundColor(FSlateColor::UseForeground())
+						.ToolTipText(LOCTEXT("MarketplaceToolTip", "Check out the Marketplace to find new projects!"))
+						.HAlign(HAlign_Center)
+						.VAlign(VAlign_Center)
 						[
-							SNew(SBorder)
-							.Visibility(this, &SProjectBrowser::GetFilterActiveOverlayVisibility)
-							.BorderImage(FEditorStyle::Get().GetBrush("SearchBox.ActiveBorder"))
+							SNew(SHorizontalBox)
+
+							+ SHorizontalBox::Slot()
+							.Padding(2.0f)
+							.VAlign(VAlign_Center)
+							.AutoWidth()
+							[
+								SNew(SImage)
+								.Image(FEditorStyle::GetBrush("LevelEditor.OpenMarketplace.Small")) 
+							]
+
+							+ SHorizontalBox::Slot()
+							.VAlign(VAlign_Center)
+							.Padding(2.0f)
+							[
+								SNew(STextBlock)
+								.TextStyle(FEditorStyle::Get(), "ProjectBrowser.Toolbar.Text")
+								.Text(LOCTEXT("Marketplace", "Marketplace"))
+							]
 						]
 					]
 
 					+SHorizontalBox::Slot()
 					.AutoWidth()
+					.VAlign(VAlign_Center)
 					[
 						SNew(SButton)
-						.ButtonStyle(FCoreStyle::Get(), "NoBorder")
+						.ButtonStyle(FEditorStyle::Get(), "ToggleButton")
 						.OnClicked(this, &SProjectBrowser::FindProjects)
+						.ForegroundColor(FSlateColor::UseForeground())
 						.ToolTipText(LOCTEXT("RefreshProjectList", "Refresh the project list"))
+						.HAlign(HAlign_Center)
+						.VAlign(VAlign_Center)
 						[
-							SNew(SImage).Image(FEditorStyle::Get().GetBrush("Icons.Refresh"))
+							SNew(SHorizontalBox)
+
+							+ SHorizontalBox::Slot()
+							.Padding(2.0f)
+							.VAlign(VAlign_Center)
+							.AutoWidth()
+							[
+								SNew(SImage)
+								.Image(FEditorStyle::GetBrush("Icons.Refresh")) 
+							]
+
+							+ SHorizontalBox::Slot()
+							.VAlign(VAlign_Center)
+							.Padding(2.0f)
+							[
+								SNew(STextBlock)
+								.TextStyle(FEditorStyle::Get(), "ProjectBrowser.Toolbar.Text")
+								.Text(LOCTEXT("RefreshProjectsText", "Refresh"))
+							]
 						]
 					]
 				]
 
 				+ SVerticalBox::Slot()
+				.AutoHeight()
+				.VAlign(VAlign_Top)
+				.Padding(FMargin(0, 5.f))
+				[
+					SNew(SOverlay)
+
+					+SOverlay::Slot()
+					[
+						SNew(SSearchBox)
+						.HintText(LOCTEXT("FilterHint", "Filter Projects..."))
+						.OnTextChanged(this, &SProjectBrowser::OnFilterTextChanged)
+					]
+
+					+SOverlay::Slot()
+					[
+						SNew(SBorder)
+						.Visibility(this, &SProjectBrowser::GetFilterActiveOverlayVisibility)
+						.BorderImage(FEditorStyle::Get().GetBrush("SearchBox.ActiveBorder"))
+					]
+				]
+
+				+ SVerticalBox::Slot()
+				.Padding(FMargin(0, 5.f))
 				[
 					SNew(SScrollBox)
 
@@ -198,48 +251,48 @@ void SProjectBrowser::Construct( const FArguments& InArgs )
 					]
 				]
 			]
-		]
 
-		+ SVerticalBox::Slot()
-		.Padding( 0, 40, 0, 0 )	// Lots of vertical padding before the dialog buttons at the bottom
-		.AutoHeight()
-		[
-			SNew( SHorizontalBox )
-
-			// Auto-load project
-			+ SHorizontalBox::Slot()
-			.FillWidth(1.0f)
-			.Padding(4, 0)
+			+ SVerticalBox::Slot()
+			.Padding( 0, 40, 0, 0 )	// Lots of vertical padding before the dialog buttons at the bottom
+			.AutoHeight()
 			[
-				SNew(SCheckBox)			
-				.IsChecked(GEditor->GetGameAgnosticSettings().bLoadTheMostRecentlyLoadedProjectAtStartup ? ESlateCheckBoxState::Checked : ESlateCheckBoxState::Unchecked)
-				.OnCheckStateChanged(this, &SProjectBrowser::OnAutoloadLastProjectChanged)
-				.Content()
+				SNew( SHorizontalBox )
+
+				// Auto-load project
+				+ SHorizontalBox::Slot()
+				.FillWidth(1.0f)
+				.VAlign(VAlign_Center)
 				[
-					SNew(STextBlock).Text(LOCTEXT("AutoloadOnStartupCheckbox", "Always load last project on startup"))
+					SNew(SCheckBox)			
+					.IsChecked(GEditor->GetGameAgnosticSettings().bLoadTheMostRecentlyLoadedProjectAtStartup ? ESlateCheckBoxState::Checked : ESlateCheckBoxState::Unchecked)
+					.OnCheckStateChanged(this, &SProjectBrowser::OnAutoloadLastProjectChanged)
+					.Content()
+					[
+						SNew(STextBlock).Text(LOCTEXT("AutoloadOnStartupCheckbox", "Always load last project on startup"))
+					]
 				]
-			]
 
-			// Browse Button
-			+SHorizontalBox::Slot()
-			.AutoWidth()
-			.Padding(4, 0)
-			[
-				SNew(SButton)
-				.Text(LOCTEXT("BrowseProjectButton", "Browse..."))
-				.OnClicked(this, &SProjectBrowser::OnBrowseToProjectClicked)
-				.ContentPadding( FCoreStyle::Get().GetMargin("StandardDialog.ContentPadding") )
-			]
+				// Browse Button
+				+SHorizontalBox::Slot()
+				.AutoWidth()
+				.Padding(8.f, 0)
+				[
+					SNew(SButton)
+					.Text(LOCTEXT("BrowseProjectButton", "Browse..."))
+					.OnClicked(this, &SProjectBrowser::OnBrowseToProjectClicked)
+					.ContentPadding( FCoreStyle::Get().GetMargin("StandardDialog.ContentPadding") )
+				]
 
-			// Open Button
-			+SHorizontalBox::Slot()
-			.AutoWidth()
-			[
-				SNew(SButton)
-				.Text(LOCTEXT("OpenProjectButton", "Open"))
-				.OnClicked(this, &SProjectBrowser::HandleOpenProjectButtonClicked)
-				.IsEnabled(this, &SProjectBrowser::HandleOpenProjectButtonIsEnabled)
-				.ContentPadding( FCoreStyle::Get().GetMargin("StandardDialog.ContentPadding") )
+				// Open Button
+				+SHorizontalBox::Slot()
+				.AutoWidth()
+				[
+					SNew(SButton)
+					.Text(LOCTEXT("OpenProjectButton", "Open"))
+					.OnClicked(this, &SProjectBrowser::HandleOpenProjectButtonClicked)
+					.IsEnabled(this, &SProjectBrowser::HandleOpenProjectButtonIsEnabled)
+					.ContentPadding( FCoreStyle::Get().GetMargin("StandardDialog.ContentPadding") )
+				]
 			]
 		]
 	];
@@ -272,7 +325,6 @@ void SProjectBrowser::ConstructCategory( const TSharedRef<SVerticalBox>& InCateg
 {
 	// Title
 	InCategoriesBox->AddSlot()
-	.Padding(FMargin(5.f, 0.f))
 	.AutoHeight()
 	[
 		SNew(STextBlock)
@@ -284,7 +336,7 @@ void SProjectBrowser::ConstructCategory( const TSharedRef<SVerticalBox>& InCateg
 	// Separator
 	InCategoriesBox->AddSlot()
 	.AutoHeight()
-	.Padding(5.0f, 2.0f, 5.0f, 8.0f)
+	.Padding(0, 2.0f, 0, 8.0f)
 	[
 		SNew(SSeparator)
 		.Visibility(this, &SProjectBrowser::GetProjectCategoryVisibility, Category)
@@ -293,7 +345,7 @@ void SProjectBrowser::ConstructCategory( const TSharedRef<SVerticalBox>& InCateg
 	// Project tile view
 	InCategoriesBox->AddSlot()
 	.AutoHeight()
-	.Padding(5.0f, 0.0f, 5.0f, 40.0f)
+	.Padding(0.f, 0.0f, 0.f, 40.0f)
 	[
 		SAssignNew(Category->ProjectTileView, STileView<TSharedPtr<FProjectItem> >)
 		.Visibility(this, &SProjectBrowser::GetProjectCategoryVisibility, Category)
@@ -1160,6 +1212,47 @@ void SProjectBrowser::HandleProjectViewSelectionChanged(TSharedPtr<FProjectItem>
 	}
 }
 
+
+FReply SProjectBrowser::HandleMarketplaceTabButtonClicked()
+{
+	IDesktopPlatform* DesktopPlatform = FDesktopPlatformModule::Get();
+
+	if (DesktopPlatform != nullptr)
+	{
+		TArray<FAnalyticsEventAttribute> EventAttributes;
+
+		if (DesktopPlatform->OpenLauncher(false, TEXT("-OpenMarket")))
+		{
+			EventAttributes.Add(FAnalyticsEventAttribute(TEXT("OpenSucceeded"), TEXT("TRUE")));
+		}
+		else
+		{
+			EventAttributes.Add(FAnalyticsEventAttribute(TEXT("OpenSucceeded"), TEXT("FALSE")));
+
+			if (EAppReturnType::Yes == FMessageDialog::Open(EAppMsgType::YesNo, LOCTEXT("InstallMarketplacePrompt", "The Marketplace requires the Unreal Engine Launcher, which does not seem to be installed on your computer. Would you like to install it now?")))
+			{
+				if (!DesktopPlatform->OpenLauncher(true, TEXT("-OpenMarket")))
+				{
+					EventAttributes.Add(FAnalyticsEventAttribute(TEXT("InstallSucceeded"), TEXT("FALSE")));
+					FMessageDialog::Open(EAppMsgType::Ok, FText::FromString(TEXT("Sorry, there was a problem installing the Launcher.\nPlease try to install it manually!")));
+				}
+				else
+				{
+					EventAttributes.Add(FAnalyticsEventAttribute(TEXT("InstallSucceeded"), TEXT("TRUE")));
+				}
+			}
+		}
+
+		EventAttributes.Add(FAnalyticsEventAttribute(TEXT("Source"), TEXT("ProjectBrowser")));
+		if( FEngineAnalytics::IsAvailable() )
+		{
+			FEngineAnalytics::GetProvider().RecordEvent(TEXT("Editor.Usage.OpenMarketplace"), EventAttributes);
+		}
+	}
+
+	return FReply::Handled();
+}
+
 void SProjectBrowser::OnFilterTextChanged(const FText& InText)
 {
 	ProjectItemFilter.SetRawFilterText(InText);
diff --git a/Engine/Source/Editor/GameProjectGeneration/Private/SProjectBrowser.h b/Engine/Source/Editor/GameProjectGeneration/Private/SProjectBrowser.h
--- a/Engine/Source/Editor/GameProjectGeneration/Private/SProjectBrowser.h
+++ b/Engine/Source/Editor/GameProjectGeneration/Private/SProjectBrowser.h
@@ -17,10 +17,6 @@ public:
 
 	SLATE_BEGIN_ARGS(SProjectBrowser) { }
 
-		SLATE_ARGUMENT(bool, AllowProjectCreate)
-
-		SLATE_EVENT(FNewProjectScreenRequested, OnNewProjectScreenRequested)
-
 	SLATE_END_ARGS()
 
 public:
@@ -109,6 +105,9 @@ private:
 	/** Handler for when the selection changes in the project view */
 	void HandleProjectViewSelectionChanged( TSharedPtr<FProjectItem> ProjectItem, ESelectInfo::Type SelectInfo, FText CategoryName );
 
+	/** Callback for clicking the 'Marketplace' button. */
+	FReply HandleMarketplaceTabButtonClicked( );
+
 	/** Called when the text in the filter box is changed */
 	void OnFilterTextChanged(const FText& InText);
 
@@ -149,9 +148,6 @@ private:
 	TSharedPtr<SVerticalBox> CategoriesBox;
 
 	bool bHasProjectFiles;
-
-	/** True if this UI will enable the creation of new projects, false otherwise */
-	bool bAllowProjectCreate;
 private:
 
 	// Holds a delegate that is executed when the new project screen is being requested.
diff --git a/Engine/Source/Editor/GameProjectGeneration/Private/Tests/GameProjectAutomationTests.cpp b/Engine/Source/Editor/GameProjectGeneration/Private/Tests/GameProjectAutomationTests.cpp
--- a/Engine/Source/Editor/GameProjectGeneration/Private/Tests/GameProjectAutomationTests.cpp
+++ b/Engine/Source/Editor/GameProjectGeneration/Private/Tests/GameProjectAutomationTests.cpp
@@ -116,7 +116,7 @@ bool FBuildPromotionNewProjectCreateTest::RunTest(const FString& Parameters)
 	}
 
 	FText FailReason;
-	if (GameProjectUtils::CreateProject(DesiredProjectFilename, TEXT(""), false, true, FailReason))
+	if (GameProjectUtils::CreateProject(FProjectInformation(DesiredProjectFilename, false, true), FailReason))
 	{
 		UE_LOG(LogGameProjectGenerationTests, Display, TEXT("Generated a new project: %s"), *DesiredProjectFilename);
 		UE_LOG(LogGameProjectGenerationTests, Display, TEXT("Test successful!"));
diff --git a/Engine/Source/Editor/HardwareTargeting/HardwareTargeting.Build.cs b/Engine/Source/Editor/HardwareTargeting/HardwareTargeting.Build.cs
new file mode 100644
--- /dev/null
+++ b/Engine/Source/Editor/HardwareTargeting/HardwareTargeting.Build.cs
@@ -0,0 +1,28 @@
+// Copyright 1998-2014 Epic Games, Inc. All Rights Reserved.
+
+using UnrealBuildTool;
+
+public class HardwareTargeting : ModuleRules
+{
+	public HardwareTargeting(TargetInfo Target)
+	{
+		PublicDependencyModuleNames.AddRange(
+			new string[] {
+				"Core",
+			}
+		);
+
+		PrivateDependencyModuleNames.AddRange(
+			new string[] {
+				"CoreUObject",
+				"Engine",
+				"InputCore",
+				"SlateCore",
+				"Slate",
+				"EditorStyle",
+				"UnrealEd",
+				"Settings",
+			}
+		);
+	}
+}
diff --git a/Engine/Source/Editor/HardwareTargeting/Private/HardwareTargetingModule.cpp b/Engine/Source/Editor/HardwareTargeting/Private/HardwareTargetingModule.cpp
new file mode 100644
--- /dev/null
+++ b/Engine/Source/Editor/HardwareTargeting/Private/HardwareTargetingModule.cpp
@@ -0,0 +1,140 @@
+// Copyright 1998-2014 Epic Games, Inc. All Rights Reserved.
+
+#include "HardwareTargetingPrivatePCH.h"
+#include "HardwareTargetingModule.h"
+#include "HardwareTargetingSettings.h"
+#include "Settings.h"
+#include "Internationalization.h"
+#include "SDecoratedEnumCombo.h"
+
+#define LOCTEXT_NAMESPACE "HardwareTargeting"
+
+class FHardwareTargetingModule : public IHardwareTargetingModule
+{
+	// IModuleInterface interface
+	virtual void StartupModule() override;
+	virtual void ShutdownModule() override;
+	// End of IModuleInterface interface
+	
+	/** Apply the current hardware targeting settings if they have changed */
+	virtual void ApplyHardwareTargetingSettings();
+
+	/** Make a new combo box for choosing a hardware class target */
+	virtual TSharedRef<SWidget> MakeHardwareClassTargetCombo(FOnHardwareClassChanged OnChanged, TAttribute<EHardwareClass::Type> SelectedEnum) override;
+
+	/** Make a new combo box for choosing a graphics preference */
+	virtual TSharedRef<SWidget> MakeGraphicsPresetTargetCombo(FOnGraphicsPresetChanged OnChanged, TAttribute<EGraphicsPreset::Type> SelectedEnum) override;
+};
+
+void FHardwareTargetingModule::StartupModule()
+{
+	ISettingsModule& SettingsModule = *ISettingsModule::Get();
+		
+	SettingsModule.RegisterSettings("Project", "Project", "HardwareTargeting",
+		LOCTEXT("HardwareTargetingSettingsName", "Hardware Targeting"),
+		LOCTEXT("HardwareTargetingSettingsDescription", "Options for choosing which class of hardware to target"),
+		GetMutableDefault<UHardwareTargetingSettings>()
+	);
+
+	// Apply any settings on startup if necessary
+	ApplyHardwareTargetingSettings();
+}
+
+void FHardwareTargetingModule::ShutdownModule()
+{
+		
+}
+	
+void FHardwareTargetingModule::ApplyHardwareTargetingSettings()
+{
+	UHardwareTargetingSettings* Settings = GetMutableDefault<UHardwareTargetingSettings>();
+
+	int32 EnumValue = 0;
+	bool bFoundKey = GConfig->GetInt(TEXT("AppliedHardwareTargetingSettings"),TEXT("AppliedTargetedHardwareClass"), EnumValue, GEditorIni);
+	if (!bFoundKey || EHardwareClass::Type(EnumValue) != Settings->TargetedHardwareClass)
+	{
+		// @todo: Set up project to work with Settings->TargetedHardwareClass
+
+		GConfig->SetInt(TEXT("AppliedHardwareTargetingSettings"),TEXT("AppliedTargetedHardwareClass"), Settings->TargetedHardwareClass, GEditorIni);
+	}
+
+	bFoundKey = GConfig->GetInt(TEXT("AppliedHardwareTargetingSettings"),TEXT("AppliedDefaultGraphicsPerformance"), EnumValue, GEditorIni);
+	if (!bFoundKey || EGraphicsPreset::Type(EnumValue) != Settings->DefaultGraphicsPerformance)
+	{
+		// @todo: Set up project to work with Settings->DefaultGraphicsPerformance
+
+		GConfig->SetInt(TEXT("AppliedHardwareTargetingSettings"),TEXT("AppliedDefaultGraphicsPerformance"), Settings->DefaultGraphicsPerformance, GEditorIni);
+	}
+
+	GConfig->Flush(false, GEditorIni);
+}
+
+TSharedRef<SWidget> FHardwareTargetingModule::MakeHardwareClassTargetCombo(FOnHardwareClassChanged OnChanged, TAttribute<EHardwareClass::Type> SelectedEnum)
+{
+	TArray<SDecoratedEnumCombo<EHardwareClass::Type>::FComboOption> HardwareClassInfo;
+	HardwareClassInfo.Add(SDecoratedEnumCombo<EHardwareClass::Type>::FComboOption(
+		EHardwareClass::Desktop, FSlateIcon(FEditorStyle::GetStyleSetName(), "HardwareTargeting.DesktopPlatform"), LOCTEXT("DesktopCaption", "Desktop / Console")));
+	HardwareClassInfo.Add(SDecoratedEnumCombo<EHardwareClass::Type>::FComboOption(
+		EHardwareClass::Mobile, FSlateIcon(FEditorStyle::GetStyleSetName(), "HardwareTargeting.MobilePlatform"), LOCTEXT("MobileCaption", "Mobile / Tablet")));
+
+	return SNew(SDecoratedEnumCombo<EHardwareClass::Type>, MoveTemp(HardwareClassInfo))
+		.SelectedEnum(SelectedEnum)
+		.OnEnumChanged(OnChanged);
+}
+
+TSharedRef<SWidget> FHardwareTargetingModule::MakeGraphicsPresetTargetCombo(FOnGraphicsPresetChanged OnChanged, TAttribute<EGraphicsPreset::Type> SelectedEnum)
+{
+	TArray<SDecoratedEnumCombo<EGraphicsPreset::Type>::FComboOption> GraphicsPresetInfo;
+	GraphicsPresetInfo.Add(SDecoratedEnumCombo<EGraphicsPreset::Type>::FComboOption(
+		EGraphicsPreset::Maximum, FSlateIcon(FEditorStyle::GetStyleSetName(), "HardwareTargeting.MaximumQuality"), LOCTEXT("MaximumCaption", "Maximum Quality")));
+	GraphicsPresetInfo.Add(SDecoratedEnumCombo<EGraphicsPreset::Type>::FComboOption(
+		EGraphicsPreset::Scalable, FSlateIcon(FEditorStyle::GetStyleSetName(), "HardwareTargeting.ScalableQuality"), LOCTEXT("ScalableCaption", "Scalable 3D or 2D")));
+
+	return SNew(SDecoratedEnumCombo<EGraphicsPreset::Type>, MoveTemp(GraphicsPresetInfo))
+		.SelectedEnum(SelectedEnum)
+		.OnEnumChanged(OnChanged);
+}
+
+IMPLEMENT_MODULE( FHardwareTargetingModule, HardwareTarget );
+	
+IHardwareTargetingModule& IHardwareTargetingModule::Get()
+{
+	static FHardwareTargetingModule Instance;
+	return Instance;
+}
+
+UHardwareTargetingSettings::UHardwareTargetingSettings( const FPostConstructInitializeProperties& PCIP )
+	: Super(PCIP)
+	, TargetedHardwareClass(EHardwareClass::Desktop)
+	, DefaultGraphicsPerformance(EGraphicsPreset::Maximum)
+{ }
+	
+bool UHardwareTargetingSettings::HasPendingChanges() const
+{
+	UHardwareTargetingSettings* Settings = GetMutableDefault<UHardwareTargetingSettings>();
+
+	int32 EnumValue = 0;
+	if (!GConfig->GetInt(TEXT("AppliedHardwareTargetingSettings"),TEXT("AppliedTargetedHardwareClass"), EnumValue, GEditorIni))
+	{
+		return true;
+	}
+
+	if (EHardwareClass::Type(EnumValue) != Settings->TargetedHardwareClass)
+	{
+		return true;
+	}
+
+	if (!GConfig->GetInt(TEXT("AppliedHardwareTargetingSettings"),TEXT("AppliedDefaultGraphicsPerformance"), EnumValue, GEditorIni))
+	{
+		return true;
+	}
+
+	if (EGraphicsPreset::Type(EnumValue) != Settings->DefaultGraphicsPerformance)
+	{
+		return true;
+	}
+
+	return false;
+}
+
+#undef LOCTEXT_NAMESPACE
\ No newline at end of file
diff --git a/Engine/Source/Editor/HardwareTargeting/Private/HardwareTargetingPrivatePCH.h b/Engine/Source/Editor/HardwareTargeting/Private/HardwareTargetingPrivatePCH.h
new file mode 100644
--- /dev/null
+++ b/Engine/Source/Editor/HardwareTargeting/Private/HardwareTargetingPrivatePCH.h
@@ -0,0 +1,7 @@
+// Copyright 1998-2014 Epic Games, Inc. All Rights Reserved.
+
+#pragma once
+
+#include "Engine.h"
+#include "Slate.h"
+#include "EditorStyle.h"
\ No newline at end of file
diff --git a/Engine/Source/Editor/HardwareTargeting/Private/SDecoratedEnumCombo.h b/Engine/Source/Editor/HardwareTargeting/Private/SDecoratedEnumCombo.h
new file mode 100644
--- /dev/null
+++ b/Engine/Source/Editor/HardwareTargeting/Private/SDecoratedEnumCombo.h
@@ -0,0 +1,156 @@
+// Copyright 1998-2014 Epic Games, Inc. All Rights Reserved.
+
+#pragma once
+
+template<typename TEnumType>
+class SDecoratedEnumCombo : public SCompoundWidget
+{
+public:
+	struct FComboOption
+	{
+		FComboOption(TEnumType InValue, FSlateIcon InIcon, FText InText)
+			: Value(InValue), Icon(InIcon), Text(InText)
+		{}
+
+		TEnumType Value;
+		FSlateIcon Icon;
+		FText Text;
+	};
+
+	DECLARE_DELEGATE_OneParam(FOnEnumChanged, TEnumType)
+
+	SLATE_BEGIN_ARGS( SDecoratedEnumCombo )
+		: _ContentPadding(6.f)
+	{}
+
+		SLATE_EVENT(FOnEnumChanged, OnEnumChanged)
+		SLATE_ARGUMENT(FMargin, ContentPadding)
+		SLATE_ATTRIBUTE(TEnumType, SelectedEnum)
+
+	SLATE_END_ARGS()
+
+	/** Constructs this widget with InArgs */
+	void Construct( const FArguments& InArgs, TArray<FComboOption> InOptions )
+	{
+		ContentPadding = InArgs._ContentPadding;
+		OnEnumChanged = InArgs._OnEnumChanged;
+		Options = MoveTemp(InOptions);
+		SelectedEnum = InArgs._SelectedEnum;
+
+		ChildSlot
+		[
+			SNew(SComboButton)
+			.ButtonStyle(FEditorStyle::Get(), "HoverHintOnly")
+			.ForegroundColor(FSlateColor::UseForeground())
+			.ContentPadding(ContentPadding)
+			.OnGetMenuContent(this, &SDecoratedEnumCombo::OnGetComboContent)
+			.ButtonContent()
+			[
+				SNew(SVerticalBox)
+
+				+ SVerticalBox::Slot()
+				.AutoHeight()
+				.HAlign(HAlign_Center)
+				[
+					SNew(SImage)
+					.Image(this, &SDecoratedEnumCombo::GetCurrentIcon)
+				]
+
+				+ SVerticalBox::Slot()
+				.AutoHeight()
+				.Padding(4)
+				.HAlign(HAlign_Center)
+				[
+					SNew(STextBlock)
+					.Text(this, &SDecoratedEnumCombo::GetCurrentText)
+				]
+			]
+		];
+	}
+
+private:
+
+	TArray<FComboOption> Options;
+	TAttribute<TEnumType> SelectedEnum;
+	FMargin ContentPadding;
+	FOnEnumChanged OnEnumChanged;
+	TWeakPtr<SWidget> MenuContent;
+
+	FText GetCurrentText() const
+	{
+		auto CurrentEnum = SelectedEnum.Get();
+		if (Options.IsValidIndex(CurrentEnum))
+		{
+			return Options[CurrentEnum].Text;
+		}
+		return FText();
+	}
+
+	const FSlateBrush* GetCurrentIcon() const
+	{
+		auto CurrentEnum = SelectedEnum.Get();
+		if (Options.IsValidIndex(CurrentEnum))
+		{
+			return Options[CurrentEnum].Icon.GetIcon();
+		}
+		return nullptr;
+	}
+
+	FReply OnChangeSelected(int32 NewIndex)
+	{
+		auto MenuContentPin = MenuContent.Pin();
+		if (MenuContentPin.IsValid())
+		{
+			TSharedRef<SWindow> ParentContextMenuWindow = FSlateApplication::Get().FindWidgetWindow( MenuContentPin.ToSharedRef() ).ToSharedRef();
+			FSlateApplication::Get().RequestDestroyWindow( ParentContextMenuWindow );
+		}
+
+		if (Options.IsValidIndex(NewIndex))
+		{
+			OnEnumChanged.ExecuteIfBound(Options[NewIndex].Value);
+		}
+		return FReply::Handled();
+	}
+
+	TSharedRef<SWidget> OnGetComboContent()
+	{
+		TSharedRef<SHorizontalBox> HorizontalBox = SNew(SHorizontalBox);
+
+		for (int32 Index = 0; Index < Options.Num(); ++Index)
+		{
+			const auto& Option = Options[Index];
+
+			HorizontalBox->AddSlot()
+			.AutoWidth()
+			[
+				SNew(SButton)
+				.ButtonStyle(FEditorStyle::Get(), "HoverHintOnly")
+				.ForegroundColor(FSlateColor::UseForeground())
+				.ContentPadding(ContentPadding)
+				.OnClicked(this, &SDecoratedEnumCombo::OnChangeSelected, Index)
+				[
+					SNew(SVerticalBox)
+
+					+ SVerticalBox::Slot()
+					.AutoHeight()
+					.HAlign(HAlign_Center)
+					[
+						SNew(SImage).Image(Option.Icon.GetIcon())
+					]
+
+					+ SVerticalBox::Slot()
+					.AutoHeight()
+					.Padding(4)
+					.HAlign(HAlign_Center)
+					[
+						SNew(STextBlock).Text(Option.Text)
+					]
+				]
+			];
+		}
+
+		MenuContent = HorizontalBox;
+
+		return HorizontalBox;
+	}
+};
\ No newline at end of file
diff --git a/Engine/Source/Editor/HardwareTargeting/Public/HardwareTargetingModule.h b/Engine/Source/Editor/HardwareTargeting/Public/HardwareTargetingModule.h
new file mode 100644
--- /dev/null
+++ b/Engine/Source/Editor/HardwareTargeting/Public/HardwareTargetingModule.h
@@ -0,0 +1,27 @@
+// Copyright 1998-2014 Epic Games, Inc. All Rights Reserved.
+
+#pragma once
+
+#include "ModuleManager.h"
+#include "HardwareTargetingSettings.h"
+
+DECLARE_DELEGATE_OneParam(FOnHardwareClassChanged, EHardwareClass::Type)
+DECLARE_DELEGATE_OneParam(FOnGraphicsPresetChanged, EGraphicsPreset::Type)
+
+class IHardwareTargetingModule : public IModuleInterface
+{
+public:
+
+	/** Singleton access to this module */
+	HARDWARETARGETING_API static IHardwareTargetingModule& Get();
+
+	/** Apply the current hardware targeting settings if they have changed */
+	virtual void ApplyHardwareTargetingSettings() = 0;
+
+	/** Make a new combo box for choosing a hardware class target */
+	virtual TSharedRef<SWidget> MakeHardwareClassTargetCombo(FOnHardwareClassChanged OnChanged, TAttribute<EHardwareClass::Type> SelectedEnum) = 0;
+
+	/** Make a new combo box for choosing a graphics preference */
+	virtual TSharedRef<SWidget> MakeGraphicsPresetTargetCombo(FOnGraphicsPresetChanged OnChanged, TAttribute<EGraphicsPreset::Type> SelectedEnum) = 0;
+};
+
diff --git a/Engine/Source/Editor/HardwareTargeting/Public/HardwareTargetingSettings.h b/Engine/Source/Editor/HardwareTargeting/Public/HardwareTargetingSettings.h
new file mode 100644
--- /dev/null
+++ b/Engine/Source/Editor/HardwareTargeting/Public/HardwareTargetingSettings.h
@@ -0,0 +1,44 @@
+// Copyright 1998-2014 Epic Games, Inc. All Rights Reserved.
+
+#pragma once
+
+#include "HardwareTargetingSettings.generated.h"
+
+/** Enum specifying a class of hardware */
+UENUM()
+namespace EHardwareClass
+{
+	enum Type
+	{
+		Desktop, Mobile
+	};
+}
+
+/** Enum specifying a graphics preset preference */
+UENUM()
+namespace EGraphicsPreset
+{
+	enum Type
+	{
+		Maximum, Scalable
+	};
+}
+
+/** Hardware targeting settings, stored in default config, per-project */
+UCLASS(config=Editor, defaultconfig)
+class HARDWARETARGETING_API UHardwareTargetingSettings
+	: public UObject
+{
+	GENERATED_UCLASS_BODY()
+
+	/** Enum specifying which class of hardware this game is targeting */
+	UPROPERTY(config, EditAnywhere, category=None)
+	TEnumAsByte<EHardwareClass::Type> TargetedHardwareClass;
+
+	/** Enum specifying a graphics preset to use for this game */
+	UPROPERTY(config, EditAnywhere, category=None)
+	TEnumAsByte<EGraphicsPreset::Type> DefaultGraphicsPerformance;
+
+	/** Check if these settings have any pending changes that require action */
+	bool HasPendingChanges() const;
+};
\ No newline at end of file
diff --git a/Engine/Source/Editor/MainFrame/Private/Frame/MainFrameActions.cpp b/Engine/Source/Editor/MainFrame/Private/Frame/MainFrameActions.cpp
--- a/Engine/Source/Editor/MainFrame/Private/Frame/MainFrameActions.cpp
+++ b/Engine/Source/Editor/MainFrame/Private/Frame/MainFrameActions.cpp
@@ -378,9 +378,23 @@ void FMainFrameActionCallbacks::NewProject( bool bAllowProjectOpening, bool bAll
 		return;
 	}
 
+	FText Title;
+	if (bAllowProjectOpening && bAllowProjectCreate)
+	{
+		Title = LOCTEXT( "SelectProjectWindowHeader", "Select Project");
+	}
+	else if (bAllowProjectOpening)
+	{
+		Title = LOCTEXT( "OpenProjectWindowHeader", "Open Project");
+	}
+	else
+	{
+		Title = LOCTEXT( "NewProjectWindowHeader", "New Project");
+	}
+
 	TSharedRef<SWindow> NewProjectWindow =
 		SNew(SWindow)
-		.Title(LOCTEXT( "SelectProjectWindowHeader", "Select Project"))
+		.Title(Title)
 		.ClientSize( FMainFrameModule::GetProjectBrowserWindowSize() )
 		.SizingRule( ESizingRule::UserSized )
 		.SupportsMinimize(false) .SupportsMaximize(false);
diff --git a/Engine/Source/Editor/MainFrame/Private/MainFrameModule.h b/Engine/Source/Editor/MainFrame/Private/MainFrameModule.h
--- a/Engine/Source/Editor/MainFrame/Private/MainFrameModule.h
+++ b/Engine/Source/Editor/MainFrame/Private/MainFrameModule.h
@@ -125,7 +125,7 @@ protected:
 public:
 
 	/** Get the size of the project browser window */
-	static FVector2D GetProjectBrowserWindowSize() { return FVector2D(1100, 700); }
+	static FVector2D GetProjectBrowserWindowSize() { return FVector2D(1100, 730); }
 
 private:
 
diff --git a/Engine/Source/Editor/UnrealEd/Private/Editor.cpp b/Engine/Source/Editor/UnrealEd/Private/Editor.cpp
--- a/Engine/Source/Editor/UnrealEd/Private/Editor.cpp
+++ b/Engine/Source/Editor/UnrealEd/Private/Editor.cpp
@@ -566,6 +566,7 @@ void UEditorEngine::Init(IEngineLoop* InEngineLoop)
 		FModuleManager::Get().LoadModule(TEXT("DeviceProfileEditor"));
 		FModuleManager::Get().LoadModule(TEXT("SourceCodeAccess"));
 		FModuleManager::Get().LoadModule(TEXT("BehaviorTreeEditor"));
+		FModuleManager::Get().LoadModule(TEXT("HardwareTargeting"));
 
 		if (!IsRunningCommandlet())
 		{
