Fix comment which UnrealDocTool can't parse.

diff --git a/Engine/Source/Editor/GameProjectGeneration/Classes/TemplateProjectDefs.h b/Engine/Source/Editor/GameProjectGeneration/Classes/TemplateProjectDefs.h
--- a/Engine/Source/Editor/GameProjectGeneration/Classes/TemplateProjectDefs.h
+++ b/Engine/Source/Editor/GameProjectGeneration/Classes/TemplateProjectDefs.h
@@ -90,7 +90,7 @@ class UTemplateProjectDefs : public UObject
 	UPROPERTY(config)
 	FName Category;
 
-	/** Fixes up all strings in this definitions object to replace %TEMPLATENAME% with the supplied template name and %PROJECTNAME% with the supplied project name */
+	/** Fixes up all strings in this definitions object to replace \%TEMPLATENAME\% with the supplied template name and \%PROJECTNAME\% with the supplied project name */
 	void FixupStrings(const FString& TemplateName, const FString& ProjectName);
 
 	/** Returns the display name for the current culture, or English if the current culture has no translation */
