Editor: Add tooltips to New Project Wizard sections, clarified Target Hardware choices / where you can edit them in the future, and added an explicit prompt for the destination picker to separate it out from the settings choices

diff --git a/Engine/Source/Editor/GameProjectGeneration/Private/SNewProjectWizard.cpp b/Engine/Source/Editor/GameProjectGeneration/Private/SNewProjectWizard.cpp
--- a/Engine/Source/Editor/GameProjectGeneration/Private/SNewProjectWizard.cpp
+++ b/Engine/Source/Editor/GameProjectGeneration/Private/SNewProjectWizard.cpp
@@ -10,6 +10,7 @@
 #include "GameProjectGenerationModule.h"
 #include "SWizard.h"
 #include "HardwareTargetingModule.h"
+#include "Editor/Documentation/Public/IDocumentation.h"
 
 #define LOCTEXT_NAMESPACE "NewProjectWizard"
 
@@ -352,7 +353,8 @@ void SNewProjectWizard::Construct( const FArguments& InArgs )
 						.AutoHeight()
 						[
 							SNew(STextBlock)
-							.Text(LOCTEXT("ProjectTemplateDescription", "Choose a template to use as a starting point for your new project:"))
+							.Text(LOCTEXT("ProjectTemplateDescription", "First, choose a template to use as a starting point for your new project:"))
+							.ToolTip(IDocumentation::Get()->CreateToolTip(LOCTEXT("TemplateChoiceTooltip", "A template consists of a little bit of player control logic (either as a Blueprint or in C++), input bindings, and appropriate prototyping assets."), NULL, TEXT("Shared/Editor/NewProjectWizard"), TEXT("TemplateChoice")))
 						]
 
 						+ SVerticalBox::Slot()
@@ -480,10 +482,11 @@ void SNewProjectWizard::Construct( const FArguments& InArgs )
 									
 								+ SVerticalBox::Slot()
 								.AutoHeight()
-								.Padding(FMargin(0, 0, 0, 25.f))
+								.Padding(FMargin(0, 0, 0, 15.f))
 								[
 									SNew(STextBlock)
-									.Text(LOCTEXT("ProjectSettingsDescription", "Now choose some settings for your project. Don't worry, you can still change these after the project has been created:"))
+									.Text(LOCTEXT("ProjectSettingsDescription", "Next, choose some settings for your project. Don't worry, you can choose later or change these at any time in [Project Settings - Target Hardware]:"))
+									.ToolTip(IDocumentation::Get()->CreateToolTip(LOCTEXT("HardwareTargetTooltip", "These settings will choose good defaults for a number of other settings in the project such as post-processing flags and touch input emulation using the mouse."), NULL, TEXT("Shared/Editor/NewProjectWizard"), TEXT("TargetHardware")))
 								]
 
 								+ SVerticalBox::Slot()
@@ -553,23 +556,32 @@ void SNewProjectWizard::Construct( const FArguments& InArgs )
 												]
 											]
 										]
-
-										// File path widget
-										+ SVerticalBox::Slot()
-										.AutoHeight()
-										.HAlign(HAlign_Center)
-										[
-											SNew(SFilepath)
-											.OnBrowseForFolder(this, &SNewProjectWizard::HandleBrowseButtonClicked)
-											.LabelBackgroundBrush(FEditorStyle::GetBrush("ProjectBrowser.Background"))
-											.LabelBackgroundColor(FLinearColor::White)
-											.FolderPath(this, &SNewProjectWizard::GetCurrentProjectFilePath)
-											.Name(this, &SNewProjectWizard::GetCurrentProjectFileName)
-											.OnFolderChanged(this, &SNewProjectWizard::OnCurrentProjectFilePathChanged)
-											.OnNameChanged(this, &SNewProjectWizard::OnCurrentProjectFileNameChanged)
-										]
 									]
 								]
+
+								+ SVerticalBox::Slot()
+								.AutoHeight()
+								.Padding(FMargin(0, 0, 0, 15.f))
+								[
+									SNew(STextBlock)
+									.Text(LOCTEXT("ProjectPathDescription", "Finally, choose a location for your project to be stored:"))
+									.ToolTip(IDocumentation::Get()->CreateToolTip(LOCTEXT("ProjectPathDescriptionTooltip", "All of your project content and code will be stored here."), NULL, TEXT("Shared/Editor/NewProjectWizard"), TEXT("ProjectPath")))
+								]
+
+								+ SVerticalBox::Slot()
+								.AutoHeight()
+								.HAlign(HAlign_Center)
+								[
+									// File path widget
+									SNew(SFilepath)
+									.OnBrowseForFolder(this, &SNewProjectWizard::HandleBrowseButtonClicked)
+									.LabelBackgroundBrush(FEditorStyle::GetBrush("ProjectBrowser.Background"))
+									.LabelBackgroundColor(FLinearColor::White)
+									.FolderPath(this, &SNewProjectWizard::GetCurrentProjectFilePath)
+									.Name(this, &SNewProjectWizard::GetCurrentProjectFileName)
+									.OnFolderChanged(this, &SNewProjectWizard::OnCurrentProjectFilePathChanged)
+									.OnNameChanged(this, &SNewProjectWizard::OnCurrentProjectFileNameChanged)
+								]
 							]
 						]
 					]
