Add a program to bootstrap packaged games, which checks for dependencies being installed (an option to include prerequisites in the game is in the packaging settings), then starts the engine with the right commandline for the project. Takes an icon from Build/Windows/Application.ico if present (which allows setting the icon for content-only projects) or the game executable being launched. Resource files added to code projects now use this icon instead.

diff --git a/Engine/Source/Editor/DetailCustomizations/Private/WindowsTargetSettingsDetails.cpp b/Engine/Source/Editor/DetailCustomizations/Private/WindowsTargetSettingsDetails.cpp
--- a/Engine/Source/Editor/DetailCustomizations/Private/WindowsTargetSettingsDetails.cpp
+++ b/Engine/Source/Editor/DetailCustomizations/Private/WindowsTargetSettingsDetails.cpp
@@ -97,7 +97,15 @@ static FString GetIconFilename(EImageScope::Type Scope)
 	}
 	else
 	{
-		FString Filename = FPaths::GameSourceDir() / FString(FApp::GetGameName()) / FString(TEXT("Resources")) / PlatformName / FString(FApp::GetGameName()) + TEXT(".ico");
+		FString Filename = FPaths::GameDir() / TEXT("Build/Windows/Application.ico");
+		if(!FPaths::FileExists(Filename))
+		{
+			FString LegacyFilename = FPaths::GameSourceDir() / FString(FApp::GetGameName()) / FString(TEXT("Resources")) / PlatformName / FString(FApp::GetGameName()) + TEXT(".ico");
+			if(FPaths::FileExists(LegacyFilename))
+			{
+				Filename = LegacyFilename;
+			}
+		}
 		return FPaths::ConvertRelativePathToFull(Filename);
 	}
 }
@@ -216,21 +224,7 @@ void FWindowsTargetSettingsDetails::CustomizeDetails( IDetailLayoutBuilder& Deta
 
 bool FWindowsTargetSettingsDetails::HandlePreExternalIconCopy(const FString& InChosenImage)
 {
-	bool bSucceeded = true;
-
-	// generate resource files if we dont have any yet
-	FText FailReason;
-	TArray<FString> CreatedFiles;
-	bSucceeded = FGameProjectGenerationModule::Get().UpdateCodeResourceFiles(CreatedFiles, FailReason);
-	if(!bSucceeded)
-	{
-		FNotificationInfo Info(FailReason);
-		Info.ExpireDuration = 5.0f;
-		Info.bUseLargeFont = false;
-		FSlateNotificationManager::Get().AddNotification(Info);
-	}
-
-	return bSucceeded;
+	return true;
 }
 
 
diff --git a/Engine/Source/Editor/GameProjectGeneration/Private/GameProjectGenerationModule.cpp b/Engine/Source/Editor/GameProjectGeneration/Private/GameProjectGenerationModule.cpp
--- a/Engine/Source/Editor/GameProjectGeneration/Private/GameProjectGenerationModule.cpp
+++ b/Engine/Source/Editor/GameProjectGeneration/Private/GameProjectGenerationModule.cpp
@@ -97,13 +97,6 @@ void FGameProjectGenerationModule::GetProjectSourceDirectoryInfo(int32& OutNumFi
 	GameProjectUtils::GetProjectSourceDirectoryInfo(OutNumFiles, OutDirectorySize);
 }
 
-bool FGameProjectGenerationModule::UpdateCodeResourceFiles(TArray<FString>& OutCreatedFiles, FText& OutFailReason)
-{
-	const FString GameModuleSourcePath = FPaths::GetPath(FPaths::GetProjectFilePath()) / TEXT("Source") / FApp::GetGameName();
-	return GameProjectUtils::GenerateGameResourceFiles(GameModuleSourcePath, FApp::GetGameName(), OutCreatedFiles, OutFailReason);
-}
-
-
 void FGameProjectGenerationModule::CheckAndWarnProjectFilenameValid()
 {
 	GameProjectUtils::CheckAndWarnProjectFilenameValid();
diff --git a/Engine/Source/Editor/GameProjectGeneration/Private/GameProjectUtils.cpp b/Engine/Source/Editor/GameProjectGeneration/Private/GameProjectUtils.cpp
--- a/Engine/Source/Editor/GameProjectGeneration/Private/GameProjectUtils.cpp
+++ b/Engine/Source/Editor/GameProjectGeneration/Private/GameProjectUtils.cpp
@@ -839,7 +839,7 @@ bool GameProjectUtils::GenerateProjectFromScratch(const FProjectInformation& InP
 	if ( InProjectInfo.bShouldGenerateCode )
 	{
 		// Generate basic source code files
-		if ( !GenerateBasicSourceCode(NewProjectFolder / TEXT("Source"), NewProjectName, StartupModuleNames, CreatedFiles, OutFailReason) )
+		if ( !GenerateBasicSourceCode(NewProjectFolder / TEXT("Source"), NewProjectName, NewProjectFolder, StartupModuleNames, CreatedFiles, OutFailReason) )
 		{
 			DeleteCreatedFiles(NewProjectFolder, CreatedFiles);
 			return false;
@@ -1248,16 +1248,16 @@ bool GameProjectUtils::CreateProjectFromTemplate(const FProjectInformation& InPr
 		CreatedFiles.Add(InProjectInfo.ProjectFilename);
 	}
 
-	if ( InProjectInfo.bShouldGenerateCode )
+	// Copy resources
+	const FString GameModuleSourcePath = DestFolder / TEXT("Source") / ProjectName;
+	if (GenerateGameResourceFiles(GameModuleSourcePath, ProjectName, DestFolder, InProjectInfo.bShouldGenerateCode, CreatedFiles, OutFailReason) == false)
 	{
-		// resource folder
-		const FString GameModuleSourcePath = DestFolder / TEXT("Source") / ProjectName;
-		if (GenerateGameResourceFiles(GameModuleSourcePath, ProjectName, CreatedFiles, OutFailReason) == false)
-		{
-			DeleteCreatedFiles(DestFolder, CreatedFiles);
-			return false;
-		}
+		DeleteCreatedFiles(DestFolder, CreatedFiles);
+		return false;
+	}
 
+	if ( InProjectInfo.bShouldGenerateCode )
+	{
 		// Generate project files
 		if ( !GenerateCodeProjectFiles(InProjectInfo.ProjectFilename, OutFailReason) )
 		{
@@ -1501,7 +1501,7 @@ bool GameProjectUtils::GenerateConfigFiles(const FProjectInformation& InProjectI
 	return true;
 }
 
-bool GameProjectUtils::GenerateBasicSourceCode(const FString& NewProjectSourcePath, const FString& NewProjectName, TArray<FString>& OutGeneratedStartupModuleNames, TArray<FString>& OutCreatedFiles, FText& OutFailReason)
+bool GameProjectUtils::GenerateBasicSourceCode(const FString& NewProjectSourcePath, const FString& NewProjectName, const FString& NewProjectRoot, TArray<FString>& OutGeneratedStartupModuleNames, TArray<FString>& OutCreatedFiles, FText& OutFailReason)
 {
 	const FString GameModulePath = NewProjectSourcePath / NewProjectName;
 	const FString EditorName = NewProjectName + TEXT("Editor");
@@ -1527,7 +1527,7 @@ bool GameProjectUtils::GenerateBasicSourceCode(const FString& NewProjectSourcePa
 	}
 
 	// MyGame resource folder
-	if (GenerateGameResourceFiles(GameModulePath, NewProjectName, OutCreatedFiles, OutFailReason) == false)
+	if (GenerateGameResourceFiles(GameModulePath, NewProjectName, NewProjectRoot, true, OutCreatedFiles, OutFailReason) == false)
 	{
 		return false;
 	}
@@ -2268,63 +2268,60 @@ bool GameProjectUtils::GenerateGameModuleTargetFile(const FString& NewBuildFileN
 	return WriteOutputFile(NewBuildFileName, FinalOutput, OutFailReason);
 }
 
-bool GameProjectUtils::GenerateGameResourceFile(const FString& NewResourceFolderName, const FString& TemplateFilename, const FString& GameName, TArray<FString>& OutCreatedFiles, FText& OutFailReason)
+bool GameProjectUtils::GenerateGameResourceFiles(const FString& NewResourceFolderName, const FString& GameName, const FString& GameRoot, bool bShouldGenerateCode, TArray<FString>& OutCreatedFiles, FText& OutFailReason)
 {
-	FString Template;
-	if (!ReadTemplateFile(TemplateFilename, Template, OutFailReason))
+#if PLATFORM_WINDOWS
+	// Copy the icon if it doesn't already exist. If we're upgrading a content-only project to code, it will already have one unless it was created before content-only project icons were supported.
+	FString IconFileName = GameRoot / TEXT("Build/Windows/Application.ico");
+	if(!FPaths::FileExists(IconFileName))
 	{
-		return false;
+		if(!SourceControlHelpers::CopyFileUnderSourceControl(IconFileName, FPaths::EngineContentDir() / TEXT("Editor/Templates/Resources/Windows/_GAME_NAME_.ico"), LOCTEXT("IconFileDescription", "icon"), OutFailReason))
+		{
+			return false;
+		}
+		OutCreatedFiles.Add(IconFileName);
 	}
-
-	FString FinalOutput = Template.Replace(TEXT("%GAME_NAME%"), *GameName, ESearchCase::CaseSensitive);
 	
-	FString OutputFilename = TemplateFilename.Replace(TEXT("_GAME_NAME_"), *GameName);
-	FString FullOutputFilename = NewResourceFolderName / OutputFilename;
-
-	struct Local
+	// Generate a RC script if it's a code project
+	if(bShouldGenerateCode)
 	{
-		static bool WriteFile(const FString& InDestFile, const FText& InFileDescription, FText& OutFailureReason, FString* InFileContents, TArray<FString>* OutCreatedFileList)
-		{
-			if (WriteOutputFile(InDestFile, *InFileContents, OutFailureReason))
-			{
-				OutCreatedFileList->Add(InDestFile);
-				return true;
-			}
+		FString OutputFilename = NewResourceFolderName / FString::Printf(TEXT("Resources/Windows/%s.rc"), *GameName);
 
+		FString TemplateText;
+		if (!ReadTemplateFile(TEXT("Resources/Windows/_GAME_NAME_.rc"), TemplateText, OutFailReason))
+		{
 			return false;
 		}
-	};
 
-	return SourceControlHelpers::CheckoutOrMarkForAdd(FullOutputFilename, LOCTEXT("ResourceFileDescription", "resource"), FOnPostCheckOut::CreateStatic(&Local::WriteFile, &FinalOutput, &OutCreatedFiles), OutFailReason);
-}
+		FString RelativeIconPath = IconFileName;
+		FPaths::MakePathRelativeTo(RelativeIconPath, *OutputFilename);
+		TemplateText = TemplateText.Replace(TEXT("%ICON_PATH%"), *RelativeIconPath, ESearchCase::CaseSensitive);
+		TemplateText = TemplateText.Replace(TEXT("%GAME_NAME%"), *GameName, ESearchCase::CaseSensitive);
 
-bool GameProjectUtils::GenerateGameResourceFiles(const FString& NewResourceFolderName, const FString& GameName, TArray<FString>& OutCreatedFiles, FText& OutFailReason)
-{
-	bool bSucceeded = true;
-	FString TemplateFilename;
+		struct Local
+		{
+			static bool WriteFile(const FString& InDestFile, const FText& InFileDescription, FText& OutFailureReason, FString* InFileContents, TArray<FString>* OutCreatedFileList)
+			{
+				if (WriteOutputFile(InDestFile, *InFileContents, OutFailureReason))
+				{
+					OutCreatedFileList->Add(InDestFile);
+					return true;
+				}
 
-#if PLATFORM_WINDOWS
-	FString IconPartialName = TEXT("_GAME_NAME_");
+				return false;
+			}
+		};
 
-	// Icon (just copy this)
-	TemplateFilename = FString::Printf(TEXT("Resources/Windows/%s.ico"), *IconPartialName);
-	FString FullTemplateFilename = FPaths::EngineContentDir() / TEXT("Editor") / TEXT("Templates") / TemplateFilename;
-	FString OutputFilename = TemplateFilename.Replace(*IconPartialName, *GameName);
-	FString FullOutputFilename = NewResourceFolderName / OutputFilename;
-	bSucceeded &= SourceControlHelpers::CopyFileUnderSourceControl(FullOutputFilename, FullTemplateFilename, LOCTEXT("IconFileDescription", "icon"), OutFailReason);
-	if(bSucceeded)
-	{
-		OutCreatedFiles.Add(FullOutputFilename);
+		if(!SourceControlHelpers::CheckoutOrMarkForAdd(OutputFilename, LOCTEXT("ResourceFileDescription", "resource"), FOnPostCheckOut::CreateStatic(&Local::WriteFile, &TemplateText, &OutCreatedFiles), OutFailReason))
+		{
+			return false;
+		}
 	}
-
-	// RC
-	TemplateFilename = TEXT("Resources/Windows/_GAME_NAME_.rc");
-	bSucceeded &= GenerateGameResourceFile(NewResourceFolderName, TemplateFilename, GameName, OutCreatedFiles, OutFailReason);
 #elif PLATFORM_MAC
 	//@todo MAC: Implement MAC version of these files...
 #endif
 
-	return bSucceeded;
+	return true;
 }
 
 bool GameProjectUtils::GenerateEditorModuleBuildFile(const FString& NewBuildFileName, const FString& ModuleName, const TArray<FString>& PublicDependencyModuleNames, const TArray<FString>& PrivateDependencyModuleNames, FText& OutFailReason)
@@ -2615,7 +2612,7 @@ bool GameProjectUtils::AddCodeToProject_Internal(const FString& NewClassName, co
 		const FString GameModuleName = FApp::GetGameName();
 
 		TArray<FString> StartupModuleNames;
-		if ( GenerateBasicSourceCode(SourceDir, GameModuleName, StartupModuleNames, CreatedFiles, OutFailReason) )
+		if ( GenerateBasicSourceCode(SourceDir, GameModuleName, FPaths::GameDir(), StartupModuleNames, CreatedFiles, OutFailReason) )
 		{
 			UpdateProject(&StartupModuleNames);
 		}
diff --git a/Engine/Source/Editor/GameProjectGeneration/Private/GameProjectUtils.h b/Engine/Source/Editor/GameProjectGeneration/Private/GameProjectUtils.h
--- a/Engine/Source/Editor/GameProjectGeneration/Private/GameProjectUtils.h
+++ b/Engine/Source/Editor/GameProjectGeneration/Private/GameProjectUtils.h
@@ -199,7 +199,7 @@ public:
 	static bool GenerateCodeProjectFiles(const FString& ProjectFilename, FText& OutFailReason);
 
 	/** Generates a set of resource files for a game module */
-	static bool GenerateGameResourceFiles(const FString& NewResourceFolderName, const FString& GameName, TArray<FString>& OutCreatedFiles, FText& OutFailReason);
+	static bool GenerateGameResourceFiles(const FString& NewResourceFolderName, const FString& GameName, const FString& GameRoot, bool bShouldGenerateCode, TArray<FString>& OutCreatedFiles, FText& OutFailReason);
 
 	/** Returns true if there are starter content files available for instancing into new projects. */
 	static bool IsStarterContentAvailableForNewProjects();
@@ -307,7 +307,7 @@ private:
 	static bool GenerateConfigFiles(const FProjectInformation& InProjectInfo, TArray<FString>& OutCreatedFiles, FText& OutFailReason);
 
 	/** Creates the basic source code for a new project. On failure, OutFailReason will be populated. */
-	static bool GenerateBasicSourceCode(const FString& NewProjectSourcePath, const FString& NewProjectName, TArray<FString>& OutGeneratedStartupModuleNames, TArray<FString>& OutCreatedFiles, FText& OutFailReason);
+	static bool GenerateBasicSourceCode(const FString& NewProjectSourcePath, const FString& NewProjectName, const FString& NewProjectRoot, TArray<FString>& OutGeneratedStartupModuleNames, TArray<FString>& OutCreatedFiles, FText& OutFailReason);
 
 	/** Creates the game framework source code for a new project (Pawn, GameMode, PlayerControlleR). On failure, OutFailReason will be populated. */
 	static bool GenerateGameFrameworkSourceCode(const FString& NewProjectSourcePath, const FString& NewProjectName, TArray<FString>& OutCreatedFiles, FText& OutFailReason);
@@ -345,9 +345,6 @@ private:
 	/** Generates a Target.cs file for a game module */
 	static bool GenerateGameModuleTargetFile(const FString& NewTargetFileName, const FString& ModuleName, const TArray<FString>& ExtraModuleNames, FText& OutFailReason);
 
-	/** Generates a resource file for a game module */
-	static bool GenerateGameResourceFile(const FString& NewResourceFolderName, const FString& TemplateFilename, const FString& GameName, TArray<FString>& OutCreatedFiles, FText& OutFailReason);
-
 	/** Generates a Build.cs file for a Editor module */
 	static bool GenerateEditorModuleBuildFile(const FString& NewBuildFileName, const FString& ModuleName, const TArray<FString>& PublicDependencyModuleNames, const TArray<FString>& PrivateDependencyModuleNames, FText& OutFailReason);
 
diff --git a/Engine/Source/Editor/GameProjectGeneration/Public/GameProjectGenerationModule.h b/Engine/Source/Editor/GameProjectGeneration/Public/GameProjectGenerationModule.h
--- a/Engine/Source/Editor/GameProjectGeneration/Public/GameProjectGenerationModule.h
+++ b/Engine/Source/Editor/GameProjectGeneration/Public/GameProjectGenerationModule.h
@@ -68,9 +68,6 @@ public:
 	/** Gets file and size info about the source directory */
 	virtual void GetProjectSourceDirectoryInfo(int32& OutNumFiles, int64& OutDirectorySize);
 
-	/** Update code resource files */
-	virtual bool UpdateCodeResourceFiles(TArray<FString>& OutCreatedFiles, FText& OutFailReason);
-
 	/** Warn the user if the project filename is invalid in case they renamed it outside the editor */
 	virtual void CheckAndWarnProjectFilenameValid();
 
diff --git a/Engine/Source/Editor/MainFrame/Private/Frame/MainFrameActions.cpp b/Engine/Source/Editor/MainFrame/Private/Frame/MainFrameActions.cpp
--- a/Engine/Source/Editor/MainFrame/Private/Frame/MainFrameActions.cpp
+++ b/Engine/Source/Editor/MainFrame/Private/Frame/MainFrameActions.cpp
@@ -641,6 +641,11 @@ void FMainFrameActionCallbacks::PackageProject( const FName InPlatformInfoName )
 	  }
 	}
 
+	if (PackagingSettings->IncludePrerequisites)
+	{
+		OptionalParams += TEXT(" -prereqs");
+	}
+
 	if (PackagingSettings->ForDistribution)
 	{
 		OptionalParams += TEXT(" -distribution");
diff --git a/Engine/Source/Editor/UnrealEd/Classes/Settings/ProjectPackagingSettings.h b/Engine/Source/Editor/UnrealEd/Classes/Settings/ProjectPackagingSettings.h
--- a/Engine/Source/Editor/UnrealEd/Classes/Settings/ProjectPackagingSettings.h
+++ b/Engine/Source/Editor/UnrealEd/Classes/Settings/ProjectPackagingSettings.h
@@ -61,6 +61,10 @@ public:
 	UPROPERTY(config, EditAnywhere, Category=Packaging)
 	bool UsePakFile;
 
+	/** Specifies whether to include prerequisites of packaged games, such as redistributable operating system components, whenever possible. */
+	UPROPERTY(config, EditAnywhere, Category=Packaging)
+	bool IncludePrerequisites;
+
 	/** If enabled, on Android platforms, .pak files are placed inside the APK. */
 	UPROPERTY(config, EditAnywhere, Category = Packaging)
 	bool UseOBB_InAPK;
diff --git a/Engine/Source/Editor/UnrealEd/Private/Settings/SettingsClasses.cpp b/Engine/Source/Editor/UnrealEd/Private/Settings/SettingsClasses.cpp
--- a/Engine/Source/Editor/UnrealEd/Private/Settings/SettingsClasses.cpp
+++ b/Engine/Source/Editor/UnrealEd/Private/Settings/SettingsClasses.cpp
@@ -280,6 +280,7 @@ UProjectPackagingSettings::UProjectPackagingSettings( const class FPostConstruct
 	BuildConfiguration = PPBC_Development;
 	FullRebuild = true;
 	UsePakFile = true;
+	IncludePrerequisites = true;
 	UseOBB_InAPK = false;
 	CulturesToStage.Add("en");
 }
diff --git a/Engine/Source/Programs/AutomationTool/ProjectParams.cs b/Engine/Source/Programs/AutomationTool/ProjectParams.cs
--- a/Engine/Source/Programs/AutomationTool/ProjectParams.cs
+++ b/Engine/Source/Programs/AutomationTool/ProjectParams.cs
@@ -269,6 +269,7 @@ namespace AutomationTool
 			this.Archive = InParams.Archive;
 			this.ArchiveDirectoryParam = InParams.ArchiveDirectoryParam;
 			this.Distribution = InParams.Distribution;
+			this.Prereqs = InParams.Prereqs;
             this.OBBinAPK = InParams.OBBinAPK;
             this.Prebuilt = InParams.Prebuilt;
             this.RunTimeoutSeconds = InParams.RunTimeoutSeconds;
@@ -329,7 +330,8 @@ namespace AutomationTool
 			bool? NoXGE = null,
 			bool? Package = null,
 			bool? Pak = null,
-			bool? SignedPak = null,
+			bool? Prereqs = null,
+            bool? SignedPak = null,
             bool? NullRHI = null,
             bool? FakeClient = null,
             bool? EditorTest = null,
@@ -448,6 +450,7 @@ namespace AutomationTool
 			this.Archive = GetParamValueIfNotSpecified(Command, Archive, this.Archive, "archive");
 			this.ArchiveDirectoryParam = ParseParamValueIfNotSpecified(Command, ArchiveDirectoryParam, "archivedirectory", String.Empty);
 			this.Distribution = GetParamValueIfNotSpecified(Command, Distribution, this.Distribution, "distribution");
+			this.Prereqs = GetParamValueIfNotSpecified(Command, Prereqs, this.Prereqs, "prereqs");
             this.OBBinAPK = GetParamValueIfNotSpecified(Command, OBBinAPK, this.OBBinAPK, "obbinapk");
             this.Prebuilt = GetParamValueIfNotSpecified(Command, Prebuilt, this.Prebuilt, "prebuilt");
             if (this.Prebuilt)
@@ -1096,7 +1099,10 @@ namespace AutomationTool
 		[Help("distribution", "package for distribution the project")]
 		public bool Distribution { get; set; }
 
-        [Help("obbinapk", "package with OBB data in APK assets directory")]
+		[Help("prereqs", "stage prerequisites along with the project")]
+		public bool Prereqs { get; set; }
+
+		[Help("obbinapk", "package with OBB data in APK assets directory")]
         public bool OBBinAPK {get; set; }
 
         [Help("Prebuilt", "this is a prebuilt cooked and packaged build")]
@@ -1720,6 +1726,7 @@ namespace AutomationTool
 				CommandUtils.Log("Distribution={0}", Distribution);
                 CommandUtils.Log("OBBinAPK={0}", OBBinAPK);
                 CommandUtils.Log("Prebuilt={0}", Prebuilt);
+				CommandUtils.Log("Prereqs={0}", Prereqs);
 				CommandUtils.Log("RawProjectPath={0}", RawProjectPath);
 				CommandUtils.Log("Rocket={0}", Rocket);
 				CommandUtils.Log("Run={0}", Run);
diff --git a/Engine/Source/Programs/AutomationTool/Win/Win.Automation.csproj b/Engine/Source/Programs/AutomationTool/Win/Win.Automation.csproj
--- a/Engine/Source/Programs/AutomationTool/Win/Win.Automation.csproj
+++ b/Engine/Source/Programs/AutomationTool/Win/Win.Automation.csproj
@@ -25,7 +25,7 @@
   </PropertyGroup>
   <PropertyGroup Condition=" '$(Configuration)|$(Platform)' == 'Development|AnyCPU' ">
     <DebugType>pdbonly</DebugType>
-    <Optimize>true</Optimize>
+    <Optimize>false</Optimize>
     <OutputPath>../../../../Binaries/DotNET/AutomationScripts/</OutputPath>
     <DefineConstants>TRACE</DefineConstants>
     <ErrorReport>prompt</ErrorReport>
@@ -51,6 +51,7 @@
     </Compile>
     <Compile Include="Properties\AssemblyInfo.cs" />
     <Compile Include="WinPlatform.Automation.cs" />
+    <Compile Include="WinPlatformResources.Automation.cs" />
   </ItemGroup>
   <ItemGroup>
     <ProjectReference Include="..\AutomationTool.csproj">
@@ -65,4 +66,4 @@
     </ProjectReference>
   </ItemGroup>
   <Import Project="$(MSBuildToolsPath)\Microsoft.CSharp.targets" />
-</Project>
+</Project>
\ No newline at end of file
diff --git a/Engine/Source/Programs/AutomationTool/Win/WinPlatform.Automation.cs b/Engine/Source/Programs/AutomationTool/Win/WinPlatform.Automation.cs
--- a/Engine/Source/Programs/AutomationTool/Win/WinPlatform.Automation.cs
+++ b/Engine/Source/Programs/AutomationTool/Win/WinPlatform.Automation.cs
@@ -100,26 +100,22 @@ public abstract class BaseWinPlatform : Platform
 
 			if (Exe.StartsWith(CombinePaths(SC.RuntimeProjectRootDir, "Binaries", SC.PlatformDir)))
 			{
-				// remap the project root. For Rocket executables, rename the executable to the game name.
-				if (!Params.IsCodeBasedProject && Exe == Exes[0])
+				// remap the project root.
+				string SourceFile = CombinePaths(SC.ProjectRoot, "Binaries", SC.PlatformDir, Path.GetFileName(Exe));
+				StageExecutable("exe", SC, Path.GetDirectoryName(SourceFile), Path.GetFileNameWithoutExtension(SourceFile) + ".", true, null, CommandUtils.CombinePaths(SC.RelativeProjectRootForStage, "Binaries", SC.PlatformDir), false, WorkingFileType);
+				if(Exe == Exes[0])
 				{
-					StageExecutable("exe", SC, CombinePaths(SC.ProjectRoot, "Binaries", SC.PlatformDir), Path.GetFileNameWithoutExtension(Exe) + ".", true, null, CommandUtils.CombinePaths(SC.RelativeProjectRootForStage, "Binaries", SC.PlatformDir), false, WorkingFileType, SC.ShortProjectName + ".");
-				}
-				else
-				{
-					StageExecutable("exe", SC, CombinePaths(SC.ProjectRoot, "Binaries", SC.PlatformDir), Path.GetFileNameWithoutExtension(Exe) + ".", true, null, CommandUtils.CombinePaths(SC.RelativeProjectRootForStage, "Binaries", SC.PlatformDir), false, WorkingFileType);
+					StageBootstrapExecutable(SC, SourceFile, CombinePaths(SC.RelativeProjectRootForStage, "Binaries", SC.PlatformDir, Path.GetFileName(Exe)), "");
 				}
 			}
 			else if (Exe.StartsWith(CombinePaths(SC.RuntimeRootDir, "Engine/Binaries", SC.PlatformDir)))
 			{
-				// Move the executable for non-code rocket projects into the game directory, using the game name, so it can figure out the UProject to look for and is consitent with code projects.
-				if (!Params.IsCodeBasedProject && Exe == Exes[0])
+				// keep it in the engine directory.
+				string SourceFile = CombinePaths(SC.LocalRoot, "Engine", "Binaries", SC.PlatformDir, Path.GetFileName(Exe));
+				StageExecutable("exe", SC, CombinePaths(SC.LocalRoot, "Engine/Binaries", SC.PlatformDir), Path.GetFileNameWithoutExtension(SourceFile) + ".", true, null, null, false, WorkingFileType);
+				if(Exe == Exes[0])
 				{
-					StageExecutable("exe", SC, CombinePaths(SC.LocalRoot, "Engine/Binaries", SC.PlatformDir), Path.GetFileNameWithoutExtension(Exe) + ".", true, null, CommandUtils.CombinePaths(SC.RelativeProjectRootForStage, "Binaries", SC.PlatformDir), false, WorkingFileType, SC.ShortProjectName + ".");
-				}
-				else
-				{
-					StageExecutable("exe", SC, CombinePaths(SC.LocalRoot, "Engine/Binaries", SC.PlatformDir), Path.GetFileNameWithoutExtension(Exe) + ".", true, null, null, false, WorkingFileType);
+					StageBootstrapExecutable(SC, SourceFile, CombinePaths("Engine", "Binaries", SC.PlatformDir, Path.GetFileName(Exe)), String.Format("..\\..\\..\\{0}\\{0}.uproject", SC.ShortProjectName));
 				}
 			}
 			else
@@ -130,6 +126,54 @@ public abstract class BaseWinPlatform : Platform
 			WorkingFileType = StagedFileType.DebugNonUFS;
 		}
 
+		if(Params.Prereqs)
+		{
+			string InstallerRelativePath = CombinePaths("Engine", "Extras", "Redist", "en-us");
+			SC.StageFiles(StagedFileType.NonUFS, CombinePaths(SC.LocalRoot, InstallerRelativePath), "UE4PrereqSetup.msi", false, null, InstallerRelativePath);
+		}
+	}
+
+	void StageBootstrapExecutable(DeploymentContext SC, string TargetFile, string StagedRelativeTargetPath, string StagedArguments)
+	{
+		string InputFile = CombinePaths(SC.LocalRoot, "Engine", "Binaries", SC.PlatformDir, String.Format("BootstrapPackagedGame-{0}-Shipping.exe", SC.PlatformDir));
+		if(InternalUtils.SafeFileExists(InputFile))
+		{
+			// Create the new bootstrap program
+			string ExeName = String.Format("{0}.exe", SC.ShortProjectName);
+
+			string IntermediateDir = CombinePaths(SC.ProjectRoot, "Intermediate", "Staging");
+			InternalUtils.SafeCreateDirectory(IntermediateDir);
+
+			string IntermediateFile = CombinePaths(IntermediateDir, ExeName);
+			File.Copy(InputFile, IntermediateFile, true);
+	
+			// Get the icon from the build directory if possible
+			GroupIconResource GroupIcon = null;
+			if(InternalUtils.SafeFileExists(CombinePaths(SC.ProjectRoot, "Build/Windows/Application.ico")))
+			{
+				GroupIcon = GroupIconResource.FromIco(CombinePaths(SC.ProjectRoot, "Build/Windows/Application.ico"));
+			}
+			if(GroupIcon == null)
+			{
+				GroupIcon = GroupIconResource.FromExe(TargetFile);
+			}
+
+			// Update the resources in the new file
+			using(ModuleResourceUpdate Update = new ModuleResourceUpdate(IntermediateFile, true))
+			{
+				const int IconResourceId = 101;
+				if(GroupIcon != null) Update.SetIcons(IconResourceId, GroupIcon);
+
+				const int ExecFileResourceId = 201;
+				Update.SetData(ExecFileResourceId, ResourceType.RawData, Encoding.Unicode.GetBytes(StagedRelativeTargetPath));
+
+				const int ExecArgsResourceId = 202;
+				Update.SetData(ExecArgsResourceId, ResourceType.RawData, Encoding.Unicode.GetBytes(StagedArguments));
+			}
+
+			// Copy it to the staging directory
+			SC.StageFiles(StagedFileType.NonUFS, IntermediateDir, ExeName, false, null, "");
+		}
 	}
 
 	public override string GetCookPlatform(bool bDedicatedServer, bool bIsClientOnly, string CookFlavor)
@@ -266,4 +310,4 @@ public class Win32Platform : BaseWinPlatform
 	}
 
 	public override bool IsSupported { get { return true; } }
-}
\ No newline at end of file
+}
diff --git a/Engine/Source/Programs/AutomationTool/Win/WinPlatformResources.Automation.cs b/Engine/Source/Programs/AutomationTool/Win/WinPlatformResources.Automation.cs
new file mode 100644
--- /dev/null
+++ b/Engine/Source/Programs/AutomationTool/Win/WinPlatformResources.Automation.cs
@@ -0,0 +1,279 @@
+﻿// Copyright 1998-2014 Epic Games, Inc. All Rights Reserved.
+using System;
+using System.Collections.Generic;
+using System.Linq;
+using System.Text;
+using System.IO;
+using AutomationTool;
+using UnrealBuildTool;
+using System.Runtime.InteropServices;
+
+enum ResourceType
+{
+	Icon = 3,
+	RawData = 10,
+	GroupIcon = 14,
+	Version = 16,
+}
+
+class IconResource
+{
+	public byte Width;
+	public byte Height;
+	public byte ColorCount;
+	public ushort Planes;
+	public ushort BitCount;
+	public byte[] Data;
+
+	public IconResource(byte InWidth, byte InHeight, byte InColorCount, ushort InPlanes, ushort InBitCount, byte[] InData)
+	{
+		Width = InWidth;
+		Height = InHeight;
+		ColorCount = InColorCount;
+		Planes = InPlanes;
+		BitCount = InBitCount;
+		Data = InData;
+	}
+
+	public override string ToString()
+	{
+		return String.Format("{0}x{1}@{2}bpp", Width, Height, BitCount);
+	}
+}
+
+class GroupIconResource
+{
+	public IconResource[] Icons;
+
+	public GroupIconResource(IconResource[] InIcons)
+	{
+		Icons = InIcons;
+	}
+
+	public static GroupIconResource FromIco(string FileName)
+	{
+		using(FileStream Stream = new FileStream(FileName, FileMode.Open, FileAccess.Read))
+		{
+			BinaryReader Reader = new BinaryReader(Stream);
+
+			ushort Reserved = Reader.ReadUInt16();
+			ushort Type = Reader.ReadUInt16();
+			ushort Count = Reader.ReadUInt16();
+
+			IconResource[] Icons = new IconResource[Count];
+			uint[] Offsets = new uint[Count];
+
+			for(int Idx = 0; Idx < Count; Idx++)
+			{
+				Icons[Idx] = ReadIconHeader(Reader);
+				Offsets[Idx] = Reader.ReadUInt32();
+			}
+
+			for(int Idx = 0; Idx < Icons.Length; Idx++)
+			{
+				Stream.Seek(Offsets[Idx], SeekOrigin.Begin);
+				Stream.Read(Icons[Idx].Data, 0, Icons[Idx].Data.Length);
+			}
+
+			return new GroupIconResource(Icons);
+		}
+	}
+
+	public static GroupIconResource FromExe(string FileName)
+	{
+		using(ModuleResourceLibary Module = new ModuleResourceLibary(FileName))
+		{
+			byte[] GroupData = Module.ReadFirstResource(ResourceType.GroupIcon);
+			return FromByteArray(GroupData, Module);
+		}
+	}
+
+	public static GroupIconResource FromByteArray(byte[] GroupData, ModuleResourceLibary Module)
+	{
+		using(MemoryStream Input = new MemoryStream(GroupData))
+		{
+			BinaryReader Reader = new BinaryReader(Input);
+
+			Reader.ReadUInt16(); // Reserved
+			Reader.ReadUInt16(); // Type
+			ushort Count = Reader.ReadUInt16();
+
+			IconResource[] Icons = new IconResource[Count];
+			for(int Idx = 0; Idx < Count; Idx++)
+			{
+				Icons[Idx] = ReadIconHeader(Reader);
+				int IconId = Reader.ReadUInt16();
+				Icons[Idx].Data = Module.ReadResource(IconId, ResourceType.Icon);
+			}
+			return new GroupIconResource(Icons);
+		}
+	}
+
+	public byte[] ToByteArray()
+	{
+		using(MemoryStream Output = new MemoryStream())
+		{
+			BinaryWriter Writer = new BinaryWriter(Output);
+
+			Writer.Write((ushort)0); // Reserved
+			Writer.Write((ushort)1); // Type (1 = icons)
+			Writer.Write((ushort)Icons.Length); 
+
+			for(int Idx = 0; Idx < Icons.Length; Idx++)
+			{
+				WriteIconHeader(Writer, Icons[Idx]);
+				Writer.Write((ushort)(Idx + 1));
+			}
+
+			return Output.ToArray();
+		}
+	}
+
+	public static IconResource ReadIconHeader(BinaryReader Reader)
+	{
+		byte Width = Reader.ReadByte();
+		byte Height = Reader.ReadByte();
+		byte ColorCount = Reader.ReadByte();
+		if(Reader.ReadByte() != 0) throw new InvalidDataException();
+		ushort Planes = Reader.ReadUInt16();
+		ushort BitCount = Reader.ReadUInt16();
+		byte[] Data = new byte[Reader.ReadUInt32()];
+
+		return new IconResource(Width, Height, ColorCount, Planes, BitCount, Data);
+	}
+
+	public void WriteIconHeader(BinaryWriter Writer, IconResource Icon)
+	{
+		Writer.Write(Icon.Width);
+		Writer.Write(Icon.Height);
+		Writer.Write(Icon.ColorCount);
+		Writer.Write((byte)0);
+		Writer.Write(Icon.Planes);
+		Writer.Write(Icon.BitCount);
+		Writer.Write(Icon.Data.Length);
+	}
+}
+
+class ModuleResourceLibary : IDisposable
+{
+	const uint LOAD_LIBRARY_AS_DATAFILE = 0x00000002; 
+
+	delegate bool EnumResourceNamesDelegate(IntPtr ModuleHandle, IntPtr Type, IntPtr Name, IntPtr Param);
+
+	[DllImport("kernel32.dll", SetLastError=true)]
+	extern static IntPtr LoadLibraryEx(string FileName, IntPtr FileHandle, uint Flags);
+
+	[DllImport("kernel32.dll", SetLastError=true)]
+	extern static bool FreeLibrary(IntPtr Handle);
+
+	[DllImport("kernel32.dll", SetLastError=true)]
+	static extern bool EnumResourceNames(IntPtr ModuleHandle, IntPtr Type, EnumResourceNamesDelegate EnumDelegate, IntPtr Param);
+
+	[DllImport("kernel32.dll", SetLastError=true)]
+	static extern IntPtr FindResource(IntPtr ModuleHandle, IntPtr Name, IntPtr Type);
+
+	[DllImport("kernel32.dll", SetLastError=true)]
+	static extern IntPtr LoadResource(IntPtr ModuleHandle, IntPtr ResourceHandle);
+
+	[DllImport("kernel32.dll", SetLastError=true)]
+	static extern IntPtr LockResource(IntPtr ResourceDataHandle);
+
+	[DllImport("kernel32.dll", SetLastError=true)]
+	static extern uint SizeofResource(IntPtr ModuleHandle, IntPtr ResourceDataHandle);
+
+	IntPtr ModuleHandle;
+
+	public ModuleResourceLibary(string FileName)
+	{
+		ModuleHandle = LoadLibraryEx(FileName, new IntPtr(0), LOAD_LIBRARY_AS_DATAFILE);
+	}
+
+	public void Dispose()
+	{
+		FreeLibrary(ModuleHandle);
+	}
+
+	public byte[] ReadResource(int ResourceId, ResourceType Type)
+	{
+		return ReadResourceInternal(new IntPtr(ResourceId), new IntPtr((int)Type));
+	}
+
+	public byte[] ReadFirstResource(ResourceType Type)
+	{
+		byte[] Result = null;
+		EnumResourceNames(ModuleHandle, new IntPtr((int)Type), (Handle, FoundType, Name, Param) => { Result = ReadResourceInternal(Name, FoundType); return false; }, new IntPtr(0));
+		return Result;
+	}
+
+	byte[] ReadResourceInternal(IntPtr Name, IntPtr Type)
+	{
+		IntPtr ResourceHandle = FindResource(ModuleHandle, Name, Type);
+
+		IntPtr ResourceDataHandle = LoadResource(ModuleHandle, ResourceHandle);
+		IntPtr ResourceData = LockResource(ResourceDataHandle);
+	
+		uint ResourceSize = SizeofResource(ModuleHandle, ResourceHandle);
+
+		byte[] ResourceDataArray = new byte[ResourceSize];
+		Marshal.Copy(ResourceData, ResourceDataArray, 0, (int)ResourceSize);
+
+		return ResourceDataArray;
+	}
+}
+
+class ModuleResourceUpdate : IDisposable
+{
+	[DllImport("kernel32.dll", SetLastError=true)]
+	static extern IntPtr BeginUpdateResource(string pFileName, [MarshalAs(UnmanagedType.Bool)]bool bDeleteExistingResources);
+
+	[DllImport("kernel32.dll", SetLastError=true)]
+	static extern bool UpdateResource(IntPtr hUpdate, IntPtr lpType, IntPtr lpName, ushort wLanguage, IntPtr lpData, uint cbData);
+
+	[DllImport("kernel32.dll", SetLastError=true)]
+	static extern bool EndUpdateResource(IntPtr hUpdate, bool fDiscard);
+
+	const ushort DefaultLanguage = 1033; // en-us
+
+	IntPtr UpdateHandle;
+	List<IntPtr> UnmanagedPointers = new List<IntPtr>();
+
+	public ModuleResourceUpdate(string OutputFile, bool bRemoveExisting)
+	{
+		UpdateHandle = BeginUpdateResource(OutputFile, bRemoveExisting);
+	}
+
+	public void SetData(int ResourceId, ResourceType Type, byte[] Data)
+	{
+		IntPtr UnmanagedPointer = Marshal.AllocHGlobal(Data.Length);
+		UnmanagedPointers.Add(UnmanagedPointer);
+
+		Marshal.Copy(Data, 0, UnmanagedPointer, Data.Length);
+
+		if(!UpdateResource(UpdateHandle, new IntPtr((int)Type), new IntPtr(ResourceId), DefaultLanguage, UnmanagedPointer, (uint)Data.Length))
+		{
+			throw new Exception("Couldn't update resource");
+		}
+	}
+
+	public void SetIcons(int ResourceId, GroupIconResource GroupIcon)
+	{
+		// Write the icon group resource
+		SetData(ResourceId, ResourceType.GroupIcon, GroupIcon.ToByteArray());
+
+		// Write the individual icons
+		for(int Idx = 0; Idx < GroupIcon.Icons.Length; Idx++)
+		{
+			SetData(Idx + 1, ResourceType.Icon, GroupIcon.Icons[Idx].Data);
+		}
+	}
+
+	public void Dispose()
+	{
+		EndUpdateResource(UpdateHandle, false);
+		foreach(IntPtr UnmanagedPointer in UnmanagedPointers)
+		{
+			Marshal.FreeHGlobal(UnmanagedPointer);
+		}
+		UnmanagedPointers.Clear();
+	}
+}
diff --git a/Engine/Source/Programs/UnrealBuildTool/Configuration/UEBuildModule.cs b/Engine/Source/Programs/UnrealBuildTool/Configuration/UEBuildModule.cs
--- a/Engine/Source/Programs/UnrealBuildTool/Configuration/UEBuildModule.cs
+++ b/Engine/Source/Programs/UnrealBuildTool/Configuration/UEBuildModule.cs
@@ -1759,6 +1759,7 @@ namespace UnrealBuildTool
 			Result.Config.bFasterWithoutUnity                    = bFasterWithoutUnity;
 			Result.Config.MinFilesUsingPrecompiledHeaderOverride = MinFilesUsingPrecompiledHeaderOverride;
 			Result.Config.bEnableExceptions                      = bEnableExceptions;
+			Result.Config.bUseStaticCRT							 = (Target.Rules != null && Target.Rules.bUseStaticCRT);
 			Result.Config.OutputDirectory                        = Path.Combine(Binary.Config.IntermediateDirectory, Name);
 
 			// Switch the optimization flag if we're building a game module. Also pass the definition for building in DebugGame along (see ModuleManager.h for notes).
diff --git a/Engine/Source/Programs/UnrealBuildTool/System/CPPEnvironment.cs b/Engine/Source/Programs/UnrealBuildTool/System/CPPEnvironment.cs
--- a/Engine/Source/Programs/UnrealBuildTool/System/CPPEnvironment.cs
+++ b/Engine/Source/Programs/UnrealBuildTool/System/CPPEnvironment.cs
@@ -223,6 +223,9 @@ namespace UnrealBuildTool
 		/** True if we're compiling a DLL */
 		public bool bIsBuildingDLL = false;
 
+		/** Whether we should compile using the statically-linked CRT. This is not widely supported for the whole engine, but is required for programs that need to run without dependencies. */
+		public bool bUseStaticCRT = false;
+
 		/** Whether the CLR (Common Language Runtime) support should be enabled for C++ targets (C++/CLI). */
 		public CPPCLRMode CLRMode = CPPCLRMode.CLRDisabled;
 
@@ -270,6 +273,7 @@ namespace UnrealBuildTool
 			bCreateDebugInfo                       = InCopyEnvironment.bCreateDebugInfo;
 			bIsBuildingLibrary                     = InCopyEnvironment.bIsBuildingLibrary;
 			bIsBuildingDLL                         = InCopyEnvironment.bIsBuildingDLL;
+			bUseStaticCRT						   = InCopyEnvironment.bUseStaticCRT;
 			CLRMode                                = InCopyEnvironment.CLRMode;
 			CPPIncludeInfo.IncludePaths            .UnionWith(InCopyEnvironment.CPPIncludeInfo.IncludePaths);
 			CPPIncludeInfo.SystemIncludePaths      .UnionWith(InCopyEnvironment.CPPIncludeInfo.SystemIncludePaths);
diff --git a/Engine/Source/Programs/UnrealBuildTool/System/RulesCompiler.cs b/Engine/Source/Programs/UnrealBuildTool/System/RulesCompiler.cs
--- a/Engine/Source/Programs/UnrealBuildTool/System/RulesCompiler.cs
+++ b/Engine/Source/Programs/UnrealBuildTool/System/RulesCompiler.cs
@@ -363,11 +363,17 @@ namespace UnrealBuildTool
 		public bool bUsesSlateEditorStyle = false;
 
         /// <summary>
-        // By default we use the Release C++ Runtime (CRT), even when compiling Debug builds.  This is because the Debug C++
-        // Runtime isn't very useful when debugging Unreal Engine projects, and linking against the Debug CRT libraries forces
-        // our third party library dependencies to also be compiled using the Debug CRT (and often perform more slowly.)  Often
-        // it can be inconvenient to require a separate copy of the debug versions of third party static libraries simply
-        // so that you can debug your program's code.
+		/// Forces linking against the static CRT. This is not supported across the engine due to the need for allocator implementations to be shared (for example), and TPS 
+		/// libraries to be consistent with each other, but can be used for utility programs.
+		/// </summary>
+        public bool bUseStaticCRT = false;
+
+        /// <summary>
+        /// By default we use the Release C++ Runtime (CRT), even when compiling Debug builds.  This is because the Debug C++
+        /// Runtime isn't very useful when debugging Unreal Engine projects, and linking against the Debug CRT libraries forces
+        /// our third party library dependencies to also be compiled using the Debug CRT (and often perform more slowly.)  Often
+        /// it can be inconvenient to require a separate copy of the debug versions of third party static libraries simply
+        /// so that you can debug your program's code.
         /// </summary>
         public bool bDebugBuildsActuallyUseDebugCRT = false;
 
diff --git a/Engine/Source/Programs/UnrealBuildTool/Windows/UEBuildWindows.cs b/Engine/Source/Programs/UnrealBuildTool/Windows/UEBuildWindows.cs
--- a/Engine/Source/Programs/UnrealBuildTool/Windows/UEBuildWindows.cs
+++ b/Engine/Source/Programs/UnrealBuildTool/Windows/UEBuildWindows.cs
@@ -408,23 +408,30 @@ namespace UnrealBuildTool
             {
                 // Explicitly exclude the MS C++ runtime libraries we're not using, to ensure other libraries we link with use the same
                 // runtime library as the engine.
-                if (InBuildTarget.Configuration == UnrealTargetConfiguration.Debug && BuildConfiguration.bDebugBuildsActuallyUseDebugCRT)
-                {
-                    InBuildTarget.GlobalLinkEnvironment.Config.ExcludedLibraries.Add("MSVCRT");
-                    InBuildTarget.GlobalLinkEnvironment.Config.ExcludedLibraries.Add("MSVCPRT");
-                }
-                else
-                {
-                    InBuildTarget.GlobalLinkEnvironment.Config.ExcludedLibraries.Add("MSVCRTD");
-                    InBuildTarget.GlobalLinkEnvironment.Config.ExcludedLibraries.Add("MSVCPRTD");
-                }
+				bool bUseDebugCRT = InBuildTarget.Configuration == UnrealTargetConfiguration.Debug && BuildConfiguration.bDebugBuildsActuallyUseDebugCRT;
+				if(!InBuildTarget.Rules.bUseStaticCRT || bUseDebugCRT)
+				{
+					InBuildTarget.GlobalLinkEnvironment.Config.ExcludedLibraries.Add("LIBCMT");
+					InBuildTarget.GlobalLinkEnvironment.Config.ExcludedLibraries.Add("LIBCPMT");
+				}
+				if(!InBuildTarget.Rules.bUseStaticCRT || !bUseDebugCRT)
+				{
+					InBuildTarget.GlobalLinkEnvironment.Config.ExcludedLibraries.Add("LIBCMTD");
+					InBuildTarget.GlobalLinkEnvironment.Config.ExcludedLibraries.Add("LIBCPMTD");
+				}
+				if (InBuildTarget.Rules.bUseStaticCRT || bUseDebugCRT)
+				{
+					InBuildTarget.GlobalLinkEnvironment.Config.ExcludedLibraries.Add("MSVCRT");
+					InBuildTarget.GlobalLinkEnvironment.Config.ExcludedLibraries.Add("MSVCPRT");
+				}
+				if(InBuildTarget.Rules.bUseStaticCRT || !bUseDebugCRT)
+				{
+					InBuildTarget.GlobalLinkEnvironment.Config.ExcludedLibraries.Add("MSVCRTD");
+					InBuildTarget.GlobalLinkEnvironment.Config.ExcludedLibraries.Add("MSVCPRTD");
+				}
                 InBuildTarget.GlobalLinkEnvironment.Config.ExcludedLibraries.Add("LIBC");
-                InBuildTarget.GlobalLinkEnvironment.Config.ExcludedLibraries.Add("LIBCMT");
-                InBuildTarget.GlobalLinkEnvironment.Config.ExcludedLibraries.Add("LIBCPMT");
                 InBuildTarget.GlobalLinkEnvironment.Config.ExcludedLibraries.Add("LIBCP");
                 InBuildTarget.GlobalLinkEnvironment.Config.ExcludedLibraries.Add("LIBCD");
-                InBuildTarget.GlobalLinkEnvironment.Config.ExcludedLibraries.Add("LIBCMTD");
-                InBuildTarget.GlobalLinkEnvironment.Config.ExcludedLibraries.Add("LIBCPMTD");
                 InBuildTarget.GlobalLinkEnvironment.Config.ExcludedLibraries.Add("LIBCPD");
 
                 //@todo ATL: Currently, only VSAccessor requires ATL (which is only used in editor builds)
diff --git a/Engine/Source/Programs/UnrealBuildTool/Windows/VCToolChain.cs b/Engine/Source/Programs/UnrealBuildTool/Windows/VCToolChain.cs
--- a/Engine/Source/Programs/UnrealBuildTool/Windows/VCToolChain.cs
+++ b/Engine/Source/Programs/UnrealBuildTool/Windows/VCToolChain.cs
@@ -328,13 +328,27 @@ namespace UnrealBuildTool
 			}
 
 			// Specify the appropriate runtime library based on the platform and config.
-			if( CompileEnvironment.Config.Target.Configuration == CPPTargetConfiguration.Debug && BuildConfiguration.bDebugBuildsActuallyUseDebugCRT )
+			if( CompileEnvironment.Config.bUseStaticCRT )
 			{
-				Arguments.Append(" /MDd");
+				if( CompileEnvironment.Config.Target.Configuration == CPPTargetConfiguration.Debug && BuildConfiguration.bDebugBuildsActuallyUseDebugCRT )
+				{
+					Arguments.Append(" /MTd");
+				}
+				else
+				{
+					Arguments.Append(" /MT");
+				}
 			}
 			else
 			{
-				Arguments.Append(" /MD");
+				if( CompileEnvironment.Config.Target.Configuration == CPPTargetConfiguration.Debug && BuildConfiguration.bDebugBuildsActuallyUseDebugCRT )
+				{
+					Arguments.Append(" /MDd");
+				}
+				else
+				{
+					Arguments.Append(" /MD");
+				}
 			}
 		}
 
diff --git a/Engine/Source/Programs/Windows/BootstrapPackagedGame/BootstrapPackagedGame.Build.cs b/Engine/Source/Programs/Windows/BootstrapPackagedGame/BootstrapPackagedGame.Build.cs
new file mode 100644
--- /dev/null
+++ b/Engine/Source/Programs/Windows/BootstrapPackagedGame/BootstrapPackagedGame.Build.cs
@@ -0,0 +1,11 @@
+﻿// Copyright 1998-2014 Epic Games, Inc. All Rights Reserved.
+
+using UnrealBuildTool;
+
+public class BootstrapPackagedGame : ModuleRules
+{
+	public BootstrapPackagedGame(TargetInfo Target)
+	{
+		PublicAdditionalLibraries.Add("shlwapi.lib");
+	}
+}
diff --git a/Engine/Source/Programs/Windows/BootstrapPackagedGame/BootstrapPackagedGame.Target.cs b/Engine/Source/Programs/Windows/BootstrapPackagedGame/BootstrapPackagedGame.Target.cs
new file mode 100644
--- /dev/null
+++ b/Engine/Source/Programs/Windows/BootstrapPackagedGame/BootstrapPackagedGame.Target.cs
@@ -0,0 +1,65 @@
+﻿// Copyright 1998-2014 Epic Games, Inc. All Rights Reserved.
+
+using UnrealBuildTool;
+using System.Collections.Generic;
+
+public class BootstrapPackagedGameTarget : TargetRules
+{
+	public BootstrapPackagedGameTarget(TargetInfo Target)
+	{
+		Type = TargetType.Program;
+
+		bUseStaticCRT = true;
+	}
+
+	//
+	// TargetRules interface.
+	//
+
+	public override void SetupBinaries(
+		TargetInfo Target,
+		ref List<UEBuildBinaryConfiguration> OutBuildBinaryConfigurations,
+		ref List<string> OutExtraModuleNames
+		)
+	{
+		OutBuildBinaryConfigurations.Add(
+			new UEBuildBinaryConfiguration( InType: UEBuildBinaryType.Executable, InModuleNames: new List<string>() { "BootstrapPackagedGame" })
+			);
+	}
+
+	public override bool ShouldCompileMonolithic(UnrealTargetPlatform InPlatform, UnrealTargetConfiguration InConfiguration)
+	{
+		return true;
+	}
+
+	public override void SetupGlobalEnvironment(
+		TargetInfo Target,
+		ref LinkEnvironmentConfiguration OutLinkEnvironmentConfiguration,
+		ref CPPEnvironmentConfiguration OutCPPEnvironmentConfiguration
+		)
+	{
+		BuildConfiguration.bUseUnityBuild = false;
+		BuildConfiguration.bUseSharedPCHs = false;
+		BuildConfiguration.bUseMallocProfiler = false;
+
+		// Disable all parts of the editor
+		UEBuildConfiguration.bCompileLeanAndMeanUE = true;
+		UEBuildConfiguration.bCompileICU = false;
+		UEBuildConfiguration.bBuildEditor = false;
+		UEBuildConfiguration.bBuildWithEditorOnlyData = false;
+		UEBuildConfiguration.bCompileAgainstEngine = false;
+		UEBuildConfiguration.bCompileAgainstCoreUObject = false;
+	}
+	
+    public override bool GUBP_AlwaysBuildWithTools(UnrealTargetPlatform InHostPlatform, bool bBuildingRocket, out bool bInternalToolOnly, out bool SeparateNode)
+	{
+		bInternalToolOnly = false;
+		SeparateNode = false;
+		return (InHostPlatform == UnrealTargetPlatform.Win32 || InHostPlatform == UnrealTargetPlatform.Win64);
+	}
+	
+	public override List<UnrealTargetConfiguration> GUBP_ToolConfigs(UnrealTargetPlatform InHostPlatform)
+	{
+		return new List<UnrealTargetConfiguration> { UnrealTargetConfiguration.Shipping };
+	}
+}
diff --git a/Engine/Source/Programs/Windows/BootstrapPackagedGame/Private/BootstrapPackagedGame.cpp b/Engine/Source/Programs/Windows/BootstrapPackagedGame/Private/BootstrapPackagedGame.cpp
new file mode 100644
--- /dev/null
+++ b/Engine/Source/Programs/Windows/BootstrapPackagedGame/Private/BootstrapPackagedGame.cpp
@@ -0,0 +1,161 @@
+#include "BootstrapPackagedGame.h"
+
+#define IDI_EXEC_FILE 201
+#define IDI_EXEC_ARGS 202
+
+WCHAR* ReadResourceString(HMODULE ModuleHandle, LPCWSTR Name)
+{
+	WCHAR* Result = NULL;
+
+	HRSRC ResourceHandle = FindResource(ModuleHandle, Name, RT_RCDATA);
+	if(ResourceHandle != NULL)
+	{
+		HGLOBAL AllocHandle = LoadResource(ModuleHandle, ResourceHandle);
+		if(AllocHandle != NULL)
+		{
+			WCHAR* Data = (WCHAR*)LockResource(AllocHandle);
+			DWORD DataLen = SizeofResource(ModuleHandle, ResourceHandle) / sizeof(WCHAR);
+
+			Result = new WCHAR[DataLen + 1];
+			memcpy(Result, Data, DataLen * sizeof(WCHAR));
+			Result[DataLen] = 0;
+		}
+	}
+
+	return Result;
+}
+
+int InstallMissingPrerequisites(const WCHAR* BaseDirectory)
+{
+	// Look for missing prerequisites
+	WCHAR MissingPrerequisites[1024] = { 0, };
+	if(LoadLibrary(L"MSVCP120.DLL") == NULL || LoadLibrary(L"MSVCR120.DLL") == NULL)
+	{
+		wcscat_s(MissingPrerequisites, TEXT("Microsoft Visual C++ 2013 Runtime\n"));
+	}
+	if(LoadLibrary(L"XINPUT1_3.DLL") == NULL)
+	{
+		wcscat_s(MissingPrerequisites, TEXT("DirectX Runtime\n"));
+	}
+
+	// Check if there's anything missing
+	if(MissingPrerequisites[0] != 0)
+	{
+		WCHAR MissingPrerequisitesMsg[1024];
+		wsprintf(MissingPrerequisitesMsg, L"The following component(s) are required to run this program:\n\n%s", MissingPrerequisites);
+
+		// If we don't have the installer, just notify the user and quit
+		WCHAR PrereqInstaller[MAX_PATH];
+		PathCombine(PrereqInstaller, BaseDirectory, L"Engine\\Extras\\Redist\\en-us\\UE4PrereqSetup.msi");
+		if(GetFileAttributes(PrereqInstaller) == INVALID_FILE_ATTRIBUTES)
+		{
+			MessageBox(NULL, MissingPrerequisitesMsg, NULL, MB_OK);
+			return 9001;
+		}
+
+		// Otherwise ask them if they want to install them
+		wcscat_s(MissingPrerequisitesMsg, L"\nWould you like to install them now?");
+		if(MessageBox(NULL, MissingPrerequisitesMsg, NULL, MB_YESNO) == IDNO)
+		{
+			return 9002;
+		}
+
+		// Start the installer
+		SHELLEXECUTEINFO ShellExecuteInfo;
+		ZeroMemory(&ShellExecuteInfo, sizeof(ShellExecuteInfo));
+		ShellExecuteInfo.cbSize = sizeof(ShellExecuteInfo);
+		ShellExecuteInfo.fMask = SEE_MASK_NOCLOSEPROCESS;
+		ShellExecuteInfo.nShow = SW_SHOWNORMAL;
+		ShellExecuteInfo.lpFile = PrereqInstaller;
+		if(!ShellExecuteExW(&ShellExecuteInfo))
+		{
+			return 9003;
+		}
+
+		// Wait for the process to complete, then get its exit code
+		DWORD ExitCode = 0;
+		WaitForSingleObject(ShellExecuteInfo.hProcess, INFINITE);
+		GetExitCodeProcess(ShellExecuteInfo.hProcess, &ExitCode);
+		CloseHandle(ShellExecuteInfo.hProcess);
+		if(ExitCode != 0)
+		{
+			return 9004;
+		}
+	}
+	return 0;
+}
+
+int SpawnTarget(WCHAR* CmdLine)
+{
+	STARTUPINFO StartupInfo;
+	ZeroMemory(&StartupInfo, sizeof(StartupInfo));
+	StartupInfo.cb = sizeof(StartupInfo);
+
+	PROCESS_INFORMATION ProcessInfo;
+	ZeroMemory(&ProcessInfo, sizeof(ProcessInfo));
+
+	if(!CreateProcess(NULL, CmdLine, NULL, NULL, FALSE, 0, NULL, NULL, &StartupInfo, &ProcessInfo))
+	{
+		DWORD ErrorCode = GetLastError();
+
+		WCHAR* Buffer = new WCHAR[wcslen(CmdLine) + 50];
+		wsprintf(Buffer, L"Couldn't start:\n%s\nCreateProcess() returned %x.", CmdLine, ErrorCode);
+		MessageBoxW(NULL, Buffer, NULL, MB_OK);
+		delete Buffer;
+
+		return 9005;
+	}
+
+	CloseHandle(ProcessInfo.hThread);
+	CloseHandle(ProcessInfo.hProcess);
+	return 0;
+}
+
+int WINAPI wWinMain(HINSTANCE hInstance, HINSTANCE hPrevInstance, TCHAR* CmdLine, int ShowCmd)
+{
+	(void)hPrevInstance;
+	(void)ShowCmd;
+
+	// Get the current module filename
+	WCHAR CurrentModuleFile[MAX_PATH];
+	GetModuleFileNameW(hInstance, CurrentModuleFile, sizeof(CurrentModuleFile));
+
+	// Get the base directory from the current module filename
+	WCHAR BaseDirectory[MAX_PATH];
+	PathCanonicalize(BaseDirectory, CurrentModuleFile);
+	PathRemoveFileSpec(BaseDirectory);
+
+	// Get the executable to run
+	WCHAR* ExecFile = ReadResourceString(hInstance, MAKEINTRESOURCE(IDI_EXEC_FILE));
+	if(ExecFile == NULL)
+	{
+		MessageBoxW(NULL, L"This program is used for packaged games and is not meant to be run directly.", NULL, MB_OK);
+		return 9000;
+	}
+
+	// Create a full command line for the program to run
+	WCHAR* BaseArgs = ReadResourceString(hInstance, MAKEINTRESOURCE(IDI_EXEC_ARGS));
+	WCHAR* ChildCmdLine = new WCHAR[wcslen(BaseDirectory) + wcslen(ExecFile) + wcslen(BaseArgs) + wcslen(CmdLine) + 20];
+	wsprintf(ChildCmdLine, L"\"%s\\%s\" %s %s", BaseDirectory, ExecFile, BaseArgs, CmdLine);
+	delete BaseArgs;
+	delete ExecFile;
+
+	// Install the prerequisites
+	int ExitCode = InstallMissingPrerequisites(BaseDirectory);
+	if(ExitCode != 0)
+	{
+		delete ChildCmdLine;
+		return ExitCode;
+	}
+
+	// Spawn the target executable
+	ExitCode = SpawnTarget(ChildCmdLine);
+	if(ExitCode != 0)
+	{
+		delete ChildCmdLine;
+		return ExitCode;
+	}
+
+	delete ChildCmdLine;
+	return ExitCode;
+}
diff --git a/Engine/Source/Programs/Windows/BootstrapPackagedGame/Private/BootstrapPackagedGame.h b/Engine/Source/Programs/Windows/BootstrapPackagedGame/Private/BootstrapPackagedGame.h
new file mode 100644
--- /dev/null
+++ b/Engine/Source/Programs/Windows/BootstrapPackagedGame/Private/BootstrapPackagedGame.h
@@ -0,0 +1,7 @@
+#pragma once
+
+#include <windows.h>
+#include <tchar.h>
+#include <assert.h>
+#include <stdio.h>
+#include <shlwapi.h>
