Editor: Improve the 'starter content' button in the new project dialog
- Add warning bubble to starter content button when Mobile/Tablet is picked
- Add a fake down arrow to make the button look more interactive

diff --git a/Engine/Source/Editor/GameProjectGeneration/Private/SNewProjectWizard.cpp b/Engine/Source/Editor/GameProjectGeneration/Private/SNewProjectWizard.cpp
--- a/Engine/Source/Editor/GameProjectGeneration/Private/SNewProjectWizard.cpp
+++ b/Engine/Source/Editor/GameProjectGeneration/Private/SNewProjectWizard.cpp
@@ -527,33 +527,70 @@ void SNewProjectWizard::Construct( const FArguments& InArgs )
 											+ SHorizontalBox::Slot()
 											.AutoWidth()
 											[
-												SNew(SButton)
-												.Visibility(StarterContentVisiblity)
-												.ButtonStyle(FEditorStyle::Get(), "HoverHintOnly")
-												.ForegroundColor(FSlateColor::UseForeground())
-												.ContentPadding(6.f)
-												.ToolTipText( LOCTEXT("CopyStarterContent_ToolTip", "Enable to include an additional content pack containing simple placeable meshes with basic materials and textures.\nYou can opt out of including this to create a project that only has the bare essentials for the selected project template."))
-												.OnClicked(this, &SNewProjectWizard::OnCopyStarterContentClicked)
+												SNew(SOverlay)
+												+SOverlay::Slot()
 												[
-													SNew(SVerticalBox)
-
-													+ SVerticalBox::Slot()
-													.AutoHeight()
-													.HAlign(HAlign_Center)
+													SNew(SButton)
+													.Visibility(StarterContentVisiblity)
+													.ButtonStyle(FEditorStyle::Get(), "HoverHintOnly")
+													.ForegroundColor(FSlateColor::UseForeground())
+													.ContentPadding(6.f)
+													.ToolTipText( LOCTEXT("CopyStarterContent_ToolTip", "Enable to include an additional content pack containing simple placeable meshes with basic materials and textures.\nYou can opt out of including this to create a project that only has the bare essentials for the selected project template."))
+													.OnClicked(this, &SNewProjectWizard::OnCopyStarterContentClicked)
 													[
-														SNew(SImage)
-														.Image(this, &SNewProjectWizard::GetStartContentIcon)
+														SNew(SVerticalBox)
+
+														+ SVerticalBox::Slot()
+														.AutoHeight()
+														.HAlign(HAlign_Center)
+														[
+															SNew(SImage)
+															.Image(this, &SNewProjectWizard::GetStartContentIcon)
+														]
+
+														+ SVerticalBox::Slot()
+														.AutoHeight()
+														.Padding(4)
+														.HAlign(HAlign_Center)
+														[
+															SNew(STextBlock)
+															.Text(this, &SNewProjectWizard::GetStartContentText)
+														]
 													]
+												]
 
-													+ SVerticalBox::Slot()
-													.AutoHeight()
-													.Padding(4)
-													.HAlign(HAlign_Center)
+												// Fake combo-button down arrow to make it look more interactive
+												+SOverlay::Slot()
+												.VAlign(VAlign_Center)
+												[
+													SNew(SHorizontalBox)
+													+ SHorizontalBox::Slot()
+													.FillWidth(1.0f)
 													[
-														SNew(STextBlock)
-														.Text(this, &SNewProjectWizard::GetStartContentText)
+														SNew(SSpacer)
+													]
+													+SHorizontalBox::Slot()
+													.AutoWidth()
+													.Padding(FMargin(8, 0))
+													[
+														SNew(SImage)
+														.Visibility(EVisibility::HitTestInvisible)
+														.Image(&(FCoreStyle::Get().GetWidgetStyle<FComboButtonStyle>("ComboButton").DownArrowImage))
 													]
 												]
+
+												// Warning when enabled for mobile, since the current starter content is bad for mobile
+												+SOverlay::Slot()
+												//.Visibility(EVisibility::SelfHitTestInvisible)
+												.HAlign(HAlign_Right)
+												.VAlign(VAlign_Top)
+												.Padding(4)
+												[
+													SNew(SImage)
+													.Image(FEditorStyle::GetBrush("Icons.Warning"))
+													.ToolTipText(this, &SNewProjectWizard::GetStarterContentWarningTooltip)
+													.Visibility(this, &SNewProjectWizard::GetStarterContentWarningVisibility)
+												]
 											]
 										]
 									]
@@ -744,6 +781,23 @@ FReply SNewProjectWizard::OnCopyStarterContentClicked()
 	return FReply::Handled();
 }
 
+EVisibility SNewProjectWizard::GetStarterContentWarningVisibility() const
+{
+	return (bCopyStarterContent && (SelectedHardwareClassTarget == EHardwareClass::Mobile)) ? EVisibility::Visible : EVisibility::Collapsed;
+}
+
+FText SNewProjectWizard::GetStarterContentWarningTooltip() const
+{
+	if (SelectedGraphicsPreset == EGraphicsPreset::Maximum)
+	{
+		return LOCTEXT("StarterContentMobileWarning_Maximum", "Note: Starter content can increase the packaged size significantly, removing the example maps will result in only packaging content that is actually used");
+	}
+	else
+	{
+		return LOCTEXT("StarterContentMobileWarning_Scalable", "Warning: Starter content is not optimized for scalable mobile projects");
+	}
+}
+
 void SNewProjectWizard::Tick( const FGeometry& AllottedGeometry, const double InCurrentTime, const float InDeltaTime )
 {
 	SCompoundWidget::Tick(AllottedGeometry, InCurrentTime, InDeltaTime);
diff --git a/Engine/Source/Editor/GameProjectGeneration/Private/SNewProjectWizard.h b/Engine/Source/Editor/GameProjectGeneration/Private/SNewProjectWizard.h
--- a/Engine/Source/Editor/GameProjectGeneration/Private/SNewProjectWizard.h
+++ b/Engine/Source/Editor/GameProjectGeneration/Private/SNewProjectWizard.h
@@ -115,6 +115,11 @@ private:
 	/** Handler for when the "copy starter content" checkbox changes state */
 	FReply OnCopyStarterContentClicked();
 
+	/** Returns the visibility of the starter content warning */
+	EVisibility GetStarterContentWarningVisibility() const;
+
+	/** Returns the tooltip text (actual warning) of the starter content warning */
+	FText GetStarterContentWarningTooltip() const;
 private:
 
 	EHardwareClass::Type SelectedHardwareClassTarget;
diff --git a/Engine/Source/Editor/HardwareTargeting/Public/HardwareTargetingSettings.h b/Engine/Source/Editor/HardwareTargeting/Public/HardwareTargetingSettings.h
--- a/Engine/Source/Editor/HardwareTargeting/Public/HardwareTargetingSettings.h
+++ b/Engine/Source/Editor/HardwareTargeting/Public/HardwareTargetingSettings.h
@@ -10,7 +10,11 @@ namespace EHardwareClass
 {
 	enum Type
 	{
-		Desktop, Mobile
+		// Desktop or console
+		Desktop,
+		
+		// Mobile or tablet
+		Mobile
 	};
 }
 
@@ -20,14 +24,17 @@ namespace EGraphicsPreset
 {
 	enum Type
 	{
-		Maximum, Scalable
+		// Maximum quality - High-end features default to enabled
+		Maximum,
+		
+		// Scalable quality - Some features are disabled by default but can be enabled based on the actual hardware
+		Scalable
 	};
 }
 
 /** Hardware targeting settings, stored in default config, per-project */
 UCLASS(config=Editor, defaultconfig)
-class HARDWARETARGETING_API UHardwareTargetingSettings
-	: public UObject
+class HARDWARETARGETING_API UHardwareTargetingSettings : public UObject
 {
 	GENERATED_UCLASS_BODY()
 
