Added a cancel button to Convert project -> Open a copy, so that you have an option to cancel project copying/duplication process.

GitHub PR 445, authored by nejclesek.

diff --git a/Engine/Source/Editor/GameProjectGeneration/Private/GameProjectUtils.cpp b/Engine/Source/Editor/GameProjectGeneration/Private/GameProjectUtils.cpp
--- a/Engine/Source/Editor/GameProjectGeneration/Private/GameProjectUtils.cpp
+++ b/Engine/Source/Editor/GameProjectGeneration/Private/GameProjectUtils.cpp
@@ -1839,7 +1839,7 @@ bool GameProjectUtils::GetClassLocation(const FString& InPath, const FModuleCont
 	return true;
 }
 
-bool GameProjectUtils::DuplicateProjectForUpgrade( const FString& InProjectFile, FString &OutNewProjectFile )
+GameProjectUtils::EProjectDuplicateResult GameProjectUtils::DuplicateProjectForUpgrade( const FString& InProjectFile, FString& OutNewProjectFile )
 {
 	IPlatformFile &PlatformFile = FPlatformFileManager::Get().GetPlatformFile();
 
@@ -1914,31 +1914,41 @@ bool GameProjectUtils::DuplicateProjectForUpgrade( const FString& InProjectFile,
 
 	// Copy everything
 	bool bCopySucceeded = true;
-	GWarn->BeginSlowTask(LOCTEXT("CreatingCopyOfProject", "Creating copy of project..."), true);
+	bool bUserCanceled = false;
+	GWarn->BeginSlowTask(LOCTEXT("CreatingCopyOfProject", "Creating copy of project..."), true, true);
 	for(int32 Idx = 0; Idx < SourceDirectories.Num() && bCopySucceeded; Idx++)
 	{
 		FString TargetDirectory = NewDirectoryName + SourceDirectories[Idx].Mid(OldDirectoryName.Len());
-		bCopySucceeded = PlatformFile.CreateDirectory(*TargetDirectory);
+		bUserCanceled = GWarn->ReceivedUserCancel();
+		bCopySucceeded = !bUserCanceled && PlatformFile.CreateDirectory(*TargetDirectory);
 		GWarn->UpdateProgress(Idx + 1, SourceDirectories.Num() + SourceFiles.Num());
 	}
 	for(int32 Idx = 0; Idx < SourceFiles.Num() && bCopySucceeded; Idx++)
 	{
 		FString TargetFile = NewDirectoryName + SourceFiles[Idx].Mid(OldDirectoryName.Len());
-		bCopySucceeded = PlatformFile.CopyFile(*TargetFile, *SourceFiles[Idx]);
+		bUserCanceled = GWarn->ReceivedUserCancel();
+		bCopySucceeded =  !bUserCanceled && PlatformFile.CopyFile(*TargetFile, *SourceFiles[Idx]);
 		GWarn->UpdateProgress(SourceDirectories.Num() + Idx + 1, SourceDirectories.Num() + SourceFiles.Num());
 	}
 	GWarn->EndSlowTask();
 
-	// Wipe the directory if we couldn't update
+	// Wipe the directory if the user canceled or we couldn't update
 	if(!bCopySucceeded)
 	{
 		PlatformFile.DeleteDirectoryRecursively(*NewDirectoryName);
-		return false;
+		if(bUserCanceled)
+		{
+			return EProjectDuplicateResult::UserCanceled;
+		}
+		else
+		{
+			return EProjectDuplicateResult::Failed;
+		}
 	}
 
 	// Otherwise fixup the output project filename
 	OutNewProjectFile = NewDirectoryName / FPaths::GetCleanFilename(InProjectFile);
-	return true;
+	return EProjectDuplicateResult::Succeeded;
 }
 
 void GameProjectUtils::UpdateSupportedTargetPlatforms(const FName& InPlatformName, const bool bIsSupported)
diff --git a/Engine/Source/Editor/GameProjectGeneration/Private/GameProjectUtils.h b/Engine/Source/Editor/GameProjectGeneration/Private/GameProjectUtils.h
--- a/Engine/Source/Editor/GameProjectGeneration/Private/GameProjectUtils.h
+++ b/Engine/Source/Editor/GameProjectGeneration/Private/GameProjectUtils.h
@@ -146,6 +146,19 @@ public:
 		const UClass* BaseClass;
 	};
 
+	/** Used as a function return result when a project is duplicated when upgrading project's version in Convert project dialog - Open a copy */
+	enum class EProjectDuplicateResult : uint8
+	{
+		/** Function has successfully duplicated all project files */
+		Succeeded,
+
+		/** There were errors while duplicating project files */
+		Failed,
+		 
+		/** User has canceled project duplication process */
+		UserCanceled
+	};
+
 	/** Returns true if the project filename is properly formed and does not conflict with another project */
 	static bool IsValidProjectFileForCreation(const FString& ProjectFile, FText& OutFailReason);
 
@@ -246,7 +259,7 @@ public:
 	static bool GetClassLocation(const FString& InPath, const FModuleContextInfo& ModuleInfo, EClassLocation& OutClassLocation, FText* const OutFailReason = nullptr);
 
 	/** Creates a copy of a project directory in order to upgrade it. */
-	static bool DuplicateProjectForUpgrade( const FString& InProjectFile, FString &OutNewProjectFile );
+	static EProjectDuplicateResult DuplicateProjectForUpgrade( const FString& InProjectFile, FString& OutNewProjectFile );
 
 	/**
 	 * Update the list of supported target platforms based upon the parameters provided
diff --git a/Engine/Source/Editor/GameProjectGeneration/Private/SProjectBrowser.cpp b/Engine/Source/Editor/GameProjectGeneration/Private/SProjectBrowser.cpp
--- a/Engine/Source/Editor/GameProjectGeneration/Private/SProjectBrowser.cpp
+++ b/Engine/Source/Editor/GameProjectGeneration/Private/SProjectBrowser.cpp
@@ -1003,11 +1003,18 @@ bool SProjectBrowser::OpenProject( const FString& InProjectFile )
 		if(Selection == OpenCopyButton)
 		{
 			FString NewProjectFile;
-			if(!GameProjectUtils::DuplicateProjectForUpgrade(ProjectFile, NewProjectFile))
+			GameProjectUtils::EProjectDuplicateResult DuplicateResult = GameProjectUtils::DuplicateProjectForUpgrade(ProjectFile, NewProjectFile);
+
+			if (DuplicateResult == GameProjectUtils::EProjectDuplicateResult::UserCanceled)
+			{
+				return false;
+			}
+			else if (DuplicateResult == GameProjectUtils::EProjectDuplicateResult::Failed)
 			{
 				FMessageDialog::Open( EAppMsgType::Ok, LOCTEXT("ConvertProjectCopyFailed", "Couldn't copy project. Check you have sufficient hard drive space and write access to the project folder.") );
 				return false;
 			}
+
 			ProjectFile = NewProjectFile;
 		}
 		if(Selection == OpenExistingButton)
