Made starter content button on New Project dialog a real drop down.

This is more consistent with the other combos, and addresses TTP#347917 - PROJECT BROWSER: Starter Content buttons has a dropdown arrow, but no list.

diff --git a/Engine/Source/Editor/GameProjectGeneration/Private/SNewProjectWizard.cpp b/Engine/Source/Editor/GameProjectGeneration/Private/SNewProjectWizard.cpp
--- a/Engine/Source/Editor/GameProjectGeneration/Private/SNewProjectWizard.cpp
+++ b/Engine/Source/Editor/GameProjectGeneration/Private/SNewProjectWizard.cpp
@@ -10,6 +10,7 @@
 #include "GameProjectGenerationModule.h"
 #include "SWizard.h"
 #include "HardwareTargetingModule.h"
+#include "SDecoratedEnumCombo.h"
 #include "Editor/Documentation/Public/IDocumentation.h"
 
 #define LOCTEXT_NAMESPACE "NewProjectWizard"
@@ -317,6 +318,20 @@ void SNewProjectWizard::Construct( const FArguments& InArgs )
 	TSharedRef<SSeparator> Separator = SNew(SSeparator).Orientation(EOrientation::Orient_Vertical);
 	Separator->SetBorderBackgroundColor(FLinearColor::White.CopyWithNewOpacity(0.25f));
 
+	TSharedPtr<SWidget> StartContentCombo;
+	{
+		TArray<SDecoratedEnumCombo<int32>::FComboOption> StarterContentInfo;
+		StarterContentInfo.Add(SDecoratedEnumCombo<int32>::FComboOption(
+			0, FSlateIcon(FEditorStyle::GetStyleSetName(), "GameProjectDialog.NoStarterContent"), LOCTEXT("NoStarterContent", "No Starter Content")));
+		StarterContentInfo.Add(SDecoratedEnumCombo<int32>::FComboOption(
+			1, FSlateIcon(FEditorStyle::GetStyleSetName(), "GameProjectDialog.IncludeStarterContent"), LOCTEXT("IncludeStarterContent", "With Starter Content")));
+
+		StartContentCombo = SNew(SDecoratedEnumCombo<int32>, MoveTemp(StarterContentInfo))
+			.SelectedEnum(this, &SNewProjectWizard::GetCopyStarterContentIndex)
+			.OnEnumChanged(this, &SNewProjectWizard::OnSetCopyStarterContent)
+			.ToolTipText( LOCTEXT("CopyStarterContent_ToolTip", "Enable to include an additional content pack containing simple placeable meshes with basic materials and textures.\nYou can opt out of including this to create a project that only has the bare essentials for the selected project template."));
+	}
+
 	const float UniformPadding = 16.f;
 	ChildSlot
 	[
@@ -530,53 +545,7 @@ void SNewProjectWizard::Construct( const FArguments& InArgs )
 												SNew(SOverlay)
 												+SOverlay::Slot()
 												[
-													SNew(SButton)
-													.Visibility(StarterContentVisiblity)
-													.ButtonStyle(FEditorStyle::Get(), "HoverHintOnly")
-													.ForegroundColor(FSlateColor::UseForeground())
-													.ContentPadding(6.f)
-													.ToolTipText( LOCTEXT("CopyStarterContent_ToolTip", "Enable to include an additional content pack containing simple placeable meshes with basic materials and textures.\nYou can opt out of including this to create a project that only has the bare essentials for the selected project template."))
-													.OnClicked(this, &SNewProjectWizard::OnCopyStarterContentClicked)
-													[
-														SNew(SVerticalBox)
-
-														+ SVerticalBox::Slot()
-														.AutoHeight()
-														.HAlign(HAlign_Center)
-														[
-															SNew(SImage)
-															.Image(this, &SNewProjectWizard::GetStartContentIcon)
-														]
-
-														+ SVerticalBox::Slot()
-														.AutoHeight()
-														.Padding(4)
-														.HAlign(HAlign_Center)
-														[
-															SNew(STextBlock)
-															.Text(this, &SNewProjectWizard::GetStartContentText)
-														]
-													]
-												]
-
-												// Fake combo-button down arrow to make it look more interactive
-												+SOverlay::Slot()
-												.VAlign(VAlign_Center)
-												[
-													SNew(SHorizontalBox)
-													+ SHorizontalBox::Slot()
-													.FillWidth(1.0f)
-													[
-														SNew(SSpacer)
-													]
-													+SHorizontalBox::Slot()
-													.AutoWidth()
-													.Padding(FMargin(8, 0))
-													[
-														SNew(SImage)
-														.Visibility(EVisibility::HitTestInvisible)
-														.Image(&(FCoreStyle::Get().GetWidgetStyle<FComboButtonStyle>("ComboButton").DownArrowImage))
-													]
+													StartContentCombo.ToSharedRef()
 												]
 
 												// Warning when enabled for mobile, since the current starter content is bad for mobile
@@ -775,10 +744,9 @@ TSharedRef<SWidget> SNewProjectWizard::BuildCategoryTabs()
 	return TabStrip;
 }
 
-FReply SNewProjectWizard::OnCopyStarterContentClicked()
+void SNewProjectWizard::OnSetCopyStarterContent(int32 InCopyStarterContent)
 {
-	bCopyStarterContent = !bCopyStarterContent;
-	return FReply::Handled();
+	bCopyStarterContent = InCopyStarterContent != 0;
 }
 
 EVisibility SNewProjectWizard::GetStarterContentWarningVisibility() const
@@ -1568,14 +1536,4 @@ void SNewProjectWizard::SetGraphicsPreset(EGraphicsPreset::Type InGraphicsPreset
 	SelectedGraphicsPreset = InGraphicsPreset;
 }
 
-FText SNewProjectWizard::GetStartContentText() const
-{
-	return bCopyStarterContent ? LOCTEXT("IncludeStarterContent", "With Starter Content") : LOCTEXT("NoStarterContent", "No Starter Content");
-}
-
-const FSlateBrush* SNewProjectWizard::GetStartContentIcon() const
-{
-	return bCopyStarterContent ? FEditorStyle::GetBrush("GameProjectDialog.IncludeStarterContent") : FEditorStyle::GetBrush("GameProjectDialog.NoStarterContent");
-}
-
 #undef LOCTEXT_NAMESPACE
diff --git a/Engine/Source/Editor/GameProjectGeneration/Private/SNewProjectWizard.h b/Engine/Source/Editor/GameProjectGeneration/Private/SNewProjectWizard.h
--- a/Engine/Source/Editor/GameProjectGeneration/Private/SNewProjectWizard.h
+++ b/Engine/Source/Editor/GameProjectGeneration/Private/SNewProjectWizard.h
@@ -113,7 +113,10 @@ private:
 	FText GetNameAndLocationErrorLabelText() const;
 
 	/** Handler for when the "copy starter content" checkbox changes state */
-	FReply OnCopyStarterContentClicked();
+	void OnSetCopyStarterContent(int32 InCopyStarterContent);
+
+	/** Get whether we are copying starter content or not */
+	int32 GetCopyStarterContentIndex() const { return bCopyStarterContent ? 1 : 0; };
 
 	/** Returns the visibility of the starter content warning */
 	EVisibility GetStarterContentWarningVisibility() const;
@@ -130,11 +133,6 @@ private:
 	void SetGraphicsPreset(EGraphicsPreset::Type InGraphicsPreset);
 	EGraphicsPreset::Type GetGraphicsPreset() const { return SelectedGraphicsPreset; }
 
-private:
-
-	FText GetStartContentText() const;
-	const FSlateBrush* GetStartContentIcon() const;
-
 private:
 
 	/** Populates TemplateItemsSource with templates found on disk */
diff --git a/Engine/Source/Editor/HardwareTargeting/Private/SDecoratedEnumCombo.h b/Engine/Source/Editor/HardwareTargeting/Private/SDecoratedEnumCombo.h
deleted file mode 100644
diff --git a/Engine/Source/Editor/HardwareTargeting/Public/SDecoratedEnumCombo.h b/Engine/Source/Editor/HardwareTargeting/Public/SDecoratedEnumCombo.h
new file mode 100644
--- /dev/null
+++ b/Engine/Source/Editor/HardwareTargeting/Public/SDecoratedEnumCombo.h
@@ -0,0 +1,161 @@
+// Copyright 1998-2014 Epic Games, Inc. All Rights Reserved.
+
+#pragma once
+
+template<typename TEnumType>
+class SDecoratedEnumCombo : public SCompoundWidget
+{
+public:
+	struct FComboOption
+	{
+		FComboOption(TEnumType InValue, FSlateIcon InIcon, FText InText, bool bInChoosable = true)
+			: Value(InValue), Icon(InIcon), Text(InText), bChoosable(bInChoosable)
+		{}
+
+		TEnumType Value;
+		FSlateIcon Icon;
+		FText Text;
+		bool bChoosable;
+	};
+
+	DECLARE_DELEGATE_OneParam(FOnEnumChanged, TEnumType)
+
+	SLATE_BEGIN_ARGS( SDecoratedEnumCombo )
+		: _ContentPadding(6.f)
+	{}
+
+		SLATE_EVENT(FOnEnumChanged, OnEnumChanged)
+		SLATE_ARGUMENT(FMargin, ContentPadding)
+		SLATE_ATTRIBUTE(TEnumType, SelectedEnum)
+
+	SLATE_END_ARGS()
+
+	/** Constructs this widget with InArgs */
+	void Construct( const FArguments& InArgs, TArray<FComboOption> InOptions )
+	{
+		ContentPadding = InArgs._ContentPadding;
+		OnEnumChanged = InArgs._OnEnumChanged;
+		Options = MoveTemp(InOptions);
+		SelectedEnum = InArgs._SelectedEnum;
+
+		ChildSlot
+		[
+			SNew(SComboButton)
+			.ButtonStyle(FEditorStyle::Get(), "HoverHintOnly")
+			.ForegroundColor(FSlateColor::UseForeground())
+			.ContentPadding(ContentPadding)
+			.OnGetMenuContent(this, &SDecoratedEnumCombo::OnGetComboContent)
+			.ButtonContent()
+			[
+				SNew(SVerticalBox)
+
+				+ SVerticalBox::Slot()
+				.AutoHeight()
+				.HAlign(HAlign_Center)
+				[
+					SNew(SImage)
+					.Image(this, &SDecoratedEnumCombo::GetCurrentIcon)
+				]
+
+				+ SVerticalBox::Slot()
+				.AutoHeight()
+				.Padding(4)
+				.HAlign(HAlign_Center)
+				[
+					SNew(STextBlock)
+					.Text(this, &SDecoratedEnumCombo::GetCurrentText)
+				]
+			]
+		];
+	}
+
+private:
+
+	TArray<FComboOption> Options;
+	TAttribute<TEnumType> SelectedEnum;
+	FMargin ContentPadding;
+	FOnEnumChanged OnEnumChanged;
+	TWeakPtr<SWidget> MenuContent;
+
+	FText GetCurrentText() const
+	{
+		auto CurrentEnum = SelectedEnum.Get();
+		if (Options.IsValidIndex(CurrentEnum))
+		{
+			return Options[CurrentEnum].Text;
+		}
+		return FText();
+	}
+
+	const FSlateBrush* GetCurrentIcon() const
+	{
+		auto CurrentEnum = SelectedEnum.Get();
+		if (Options.IsValidIndex(CurrentEnum))
+		{
+			return Options[CurrentEnum].Icon.GetIcon();
+		}
+		return nullptr;
+	}
+
+	FReply OnChangeSelected(int32 NewIndex)
+	{
+		auto MenuContentPin = MenuContent.Pin();
+		if (MenuContentPin.IsValid())
+		{
+			TSharedRef<SWindow> ParentContextMenuWindow = FSlateApplication::Get().FindWidgetWindow( MenuContentPin.ToSharedRef() ).ToSharedRef();
+			FSlateApplication::Get().RequestDestroyWindow( ParentContextMenuWindow );
+		}
+
+		if (Options.IsValidIndex(NewIndex))
+		{
+			OnEnumChanged.ExecuteIfBound(Options[NewIndex].Value);
+		}
+		return FReply::Handled();
+	}
+
+	TSharedRef<SWidget> OnGetComboContent()
+	{
+		TSharedRef<SHorizontalBox> HorizontalBox = SNew(SHorizontalBox);
+
+		for (int32 Index = 0; Index < Options.Num(); ++Index)
+		{
+			const auto& Option = Options[Index];
+			if (!Option.bChoosable)
+			{
+				continue;
+			}
+
+			HorizontalBox->AddSlot()
+			.AutoWidth()
+			[
+				SNew(SButton)
+				.ButtonStyle(FEditorStyle::Get(), "HoverHintOnly")
+				.ForegroundColor(FSlateColor::UseForeground())
+				.ContentPadding(ContentPadding)
+				.OnClicked(this, &SDecoratedEnumCombo::OnChangeSelected, Index)
+				[
+					SNew(SVerticalBox)
+
+					+ SVerticalBox::Slot()
+					.AutoHeight()
+					.HAlign(HAlign_Center)
+					[
+						SNew(SImage).Image(Option.Icon.GetIcon())
+					]
+
+					+ SVerticalBox::Slot()
+					.AutoHeight()
+					.Padding(4)
+					.HAlign(HAlign_Center)
+					[
+						SNew(STextBlock).Text(Option.Text)
+					]
+				]
+			];
+		}
+
+		MenuContent = HorizontalBox;
+
+		return HorizontalBox;
+	}
+};
\ No newline at end of file
