Added line break iterator to text display on template project tiles

This addresses TTP#347914 - Project Browser: TP_VehicleAdv: Display name is cut off

diff --git a/Engine/Source/Editor/GameProjectGeneration/Private/SNewProjectWizard.cpp b/Engine/Source/Editor/GameProjectGeneration/Private/SNewProjectWizard.cpp
--- a/Engine/Source/Editor/GameProjectGeneration/Private/SNewProjectWizard.cpp
+++ b/Engine/Source/Editor/GameProjectGeneration/Private/SNewProjectWizard.cpp
@@ -12,6 +12,7 @@
 #include "HardwareTargetingModule.h"
 #include "SDecoratedEnumCombo.h"
 #include "Editor/Documentation/Public/IDocumentation.h"
+#include "BreakIterator.h"
 
 #define LOCTEXT_NAMESPACE "NewProjectWizard"
 
@@ -259,9 +260,10 @@ public:
 				.VAlign(VAlign_Top)
 				.Padding(FMargin(ThumbnailPadding, 0))
 				[
-					SNew(SRichTextBlock)
+					SNew(STextBlock)
 					.AutoWrapText(true)
 					.Justification(ETextJustify::Center)
+					.LineBreakPolicy(FBreakIterator::CreateCamelCaseBreakIterator())
 					//.HighlightText(this, &SNewProjectWizard::GetItemHighlightText)
 					.Text(InArgs._Item->Name)
 				]
