Fix marketplace button visibility check being inverted.

diff --git a/Engine/Source/Editor/GameProjectGeneration/Private/SProjectBrowser.cpp b/Engine/Source/Editor/GameProjectGeneration/Private/SProjectBrowser.cpp
--- a/Engine/Source/Editor/GameProjectGeneration/Private/SProjectBrowser.cpp
+++ b/Engine/Source/Editor/GameProjectGeneration/Private/SProjectBrowser.cpp
@@ -153,7 +153,7 @@ void SProjectBrowser::Construct( const FArguments& InArgs )
 					.Padding(FMargin(0,0,5,0))
 					[
 						SNew(SButton)
-						.Visibility( FDesktopPlatformModule::Get()->CanOpenLauncher(true) ? EVisibility::Collapsed : EVisibility::Visible )
+						.Visibility( FDesktopPlatformModule::Get()->CanOpenLauncher(true) ? EVisibility::Visible : EVisibility::Collapsed )
 						.ButtonStyle(FEditorStyle::Get(), "ToggleButton")
 						.OnClicked(this, &SProjectBrowser::HandleMarketplaceTabButtonClicked)
 						.ForegroundColor(FSlateColor::UseForeground())
