Tweaked the location of marketplace and refresh buttons on the project browser

diff --git a/Engine/Source/Editor/GameProjectGeneration/Private/SProjectBrowser.cpp b/Engine/Source/Editor/GameProjectGeneration/Private/SProjectBrowser.cpp
--- a/Engine/Source/Editor/GameProjectGeneration/Private/SProjectBrowser.cpp
+++ b/Engine/Source/Editor/GameProjectGeneration/Private/SProjectBrowser.cpp
@@ -129,7 +129,7 @@ void SProjectBrowser::Construct( const FArguments& InArgs )
 	[
 		SNew(SBorder)
 		.BorderImage( FEditorStyle::GetBrush("ToolPanel.GroupBorder") )
-		.Padding(8.f)
+		.Padding(FMargin(8.f, 4.f))
 		[
 			SNew(SVerticalBox)
 
@@ -143,21 +143,40 @@ void SProjectBrowser::Construct( const FArguments& InArgs )
 				+ SVerticalBox::Slot()
 				.AutoHeight()
 				.VAlign(VAlign_Top)
-				.HAlign(HAlign_Left)
-				.Padding(FMargin(0, 0, 0, 5.f))
 				[
 					SNew(SHorizontalBox)
 					
+					+ SHorizontalBox::Slot()
+					.Padding(FMargin(0, 0, 5.f, 0))
+					.VAlign(VAlign_Center)
+					[
+						SNew(SOverlay)
+
+						+SOverlay::Slot()
+						[
+							SNew(SSearchBox)
+							.HintText(LOCTEXT("FilterHint", "Filter Projects..."))
+							.OnTextChanged(this, &SProjectBrowser::OnFilterTextChanged)
+						]
+
+						+SOverlay::Slot()
+						[
+							SNew(SBorder)
+							.Visibility(this, &SProjectBrowser::GetFilterActiveOverlayVisibility)
+							.BorderImage(FEditorStyle::Get().GetBrush("SearchBox.ActiveBorder"))
+						]
+					]
+
 					+SHorizontalBox::Slot()
 					.AutoWidth()
-					.Padding(FMargin(0,0,5,0))
+					.VAlign(VAlign_Center)
+					.Padding(FMargin(0, 0, 5.f, 0))
 					[
 						SNew(SButton)
-						.Visibility( FDesktopPlatformModule::Get()->CanOpenLauncher(true) ? EVisibility::Visible : EVisibility::Collapsed )
 						.ButtonStyle(FEditorStyle::Get(), "ToggleButton")
-						.OnClicked(this, &SProjectBrowser::HandleMarketplaceTabButtonClicked)
+						.OnClicked(this, &SProjectBrowser::FindProjects)
 						.ForegroundColor(FSlateColor::UseForeground())
-						.ToolTipText(LOCTEXT("MarketplaceToolTip", "Check out the Marketplace to find new projects!"))
+						.ToolTipText(LOCTEXT("RefreshProjectList", "Refresh the project list"))
 						.HAlign(HAlign_Center)
 						.VAlign(VAlign_Center)
 						[
@@ -169,7 +188,7 @@ void SProjectBrowser::Construct( const FArguments& InArgs )
 							.AutoWidth()
 							[
 								SNew(SImage)
-								.Image(FEditorStyle::GetBrush("LevelEditor.OpenMarketplace.Small")) 
+								.Image(FEditorStyle::GetBrush("Icons.Refresh")) 
 							]
 
 							+ SHorizontalBox::Slot()
@@ -178,20 +197,20 @@ void SProjectBrowser::Construct( const FArguments& InArgs )
 							[
 								SNew(STextBlock)
 								.TextStyle(FEditorStyle::Get(), "ProjectBrowser.Toolbar.Text")
-								.Text(LOCTEXT("Marketplace", "Marketplace"))
+								.Text(LOCTEXT("RefreshProjectsText", "Refresh"))
 							]
 						]
 					]
 
 					+SHorizontalBox::Slot()
 					.AutoWidth()
-					.VAlign(VAlign_Center)
 					[
 						SNew(SButton)
+						.Visibility( FDesktopPlatformModule::Get()->CanOpenLauncher(true) ? EVisibility::Visible : EVisibility::Collapsed )
 						.ButtonStyle(FEditorStyle::Get(), "ToggleButton")
-						.OnClicked(this, &SProjectBrowser::FindProjects)
+						.OnClicked(this, &SProjectBrowser::HandleMarketplaceTabButtonClicked)
 						.ForegroundColor(FSlateColor::UseForeground())
-						.ToolTipText(LOCTEXT("RefreshProjectList", "Refresh the project list"))
+						.ToolTipText(LOCTEXT("MarketplaceToolTip", "Check out the Marketplace to find new projects!"))
 						.HAlign(HAlign_Center)
 						.VAlign(VAlign_Center)
 						[
@@ -203,7 +222,7 @@ void SProjectBrowser::Construct( const FArguments& InArgs )
 							.AutoWidth()
 							[
 								SNew(SImage)
-								.Image(FEditorStyle::GetBrush("Icons.Refresh")) 
+								.Image(FEditorStyle::GetBrush("LevelEditor.OpenMarketplace.Small")) 
 							]
 
 							+ SHorizontalBox::Slot()
@@ -212,34 +231,12 @@ void SProjectBrowser::Construct( const FArguments& InArgs )
 							[
 								SNew(STextBlock)
 								.TextStyle(FEditorStyle::Get(), "ProjectBrowser.Toolbar.Text")
-								.Text(LOCTEXT("RefreshProjectsText", "Refresh"))
+								.Text(LOCTEXT("Marketplace", "Marketplace"))
 							]
 						]
 					]
 				]
 
-				+ SVerticalBox::Slot()
-				.AutoHeight()
-				.VAlign(VAlign_Top)
-				.Padding(FMargin(0, 5.f))
-				[
-					SNew(SOverlay)
-
-					+SOverlay::Slot()
-					[
-						SNew(SSearchBox)
-						.HintText(LOCTEXT("FilterHint", "Filter Projects..."))
-						.OnTextChanged(this, &SProjectBrowser::OnFilterTextChanged)
-					]
-
-					+SOverlay::Slot()
-					[
-						SNew(SBorder)
-						.Visibility(this, &SProjectBrowser::GetFilterActiveOverlayVisibility)
-						.BorderImage(FEditorStyle::Get().GetBrush("SearchBox.ActiveBorder"))
-					]
-				]
-
 				+ SVerticalBox::Slot()
 				.Padding(FMargin(0, 5.f))
 				[
