Revised behaviour of GenerateConfigFiles so that it will extract the EditorStartupMap and GameDefaultMap from the ini files in the startercontent config folder if they exist

diff --git a/Engine/Source/Editor/GameProjectGeneration/Private/GameProjectUtils.cpp b/Engine/Source/Editor/GameProjectGeneration/Private/GameProjectUtils.cpp
--- a/Engine/Source/Editor/GameProjectGeneration/Private/GameProjectUtils.cpp
+++ b/Engine/Source/Editor/GameProjectGeneration/Private/GameProjectUtils.cpp
@@ -1439,31 +1439,75 @@ bool GameProjectUtils::GenerateConfigFiles(const FProjectInformation& InProjectI
 		
 		if (InProjectInfo.bCopyStarterContent)
 		{
-			// for generated/blank projects with starter content, set startup map to be the starter content map
-			// otherwise, we leave it to be what the template wants.
+			FString StarterContentContentDir = FPaths::StarterContentDir() + TEXT("Content/");
+
 			TArray<FString> StarterContentMapFiles;
 			const FString FileWildcard = FString(TEXT("*")) + FPackageName::GetMapPackageExtension();
 		
-			// assume the first map in the /Maps folder is the default map
-			IFileManager::Get().FindFilesRecursive(StarterContentMapFiles, *FPaths::StarterContentDir(), *FileWildcard, /*Files=*/true, /*Directories=*/false);
-			if (StarterContentMapFiles.Num() > 0)
+			FString SpecificEditorStartupMap;
+			FString SpecificGameDefaultMap;	
+			FString FullEditorStartupMapPath;
+			FString FullGameDefaultMapPath;
+
+			// First we check if there are maps specified in the DefaultEngine.ini in our starter content folder			
+			const FString StarterContentDefaultEngineIniFilename = FPaths::StarterContentDir() / TEXT("Config/DefaultEngine.ini");
+			if (FPaths::FileExists(StarterContentDefaultEngineIniFilename))
 			{
-				FString StarterContentContentDir = FPaths::StarterContentDir() + TEXT("Content/");
+				FString StarterFileContents;
+				if (FFileHelper::LoadFileToString(StarterFileContents, *StarterContentDefaultEngineIniFilename))
+				{
+					TArray<FString> StarterIniLines;
+					StarterFileContents.ParseIntoArrayLines(&StarterIniLines);
+					for (int32 Line = 0; Line < StarterIniLines.Num();Line++)
+					{
+						FString EachLine = StarterIniLines[Line];
+						if (EachLine.StartsWith(TEXT("EditorStartupMap")))
+						{
+							EachLine.Split("=", nullptr, &SpecificEditorStartupMap);
+							FullEditorStartupMapPath = (StarterContentContentDir / SpecificEditorStartupMap) + FPackageName::GetMapPackageExtension();
+							FullEditorStartupMapPath = FullEditorStartupMapPath.Replace(TEXT("Game/"), TEXT(""));
+						}
+						if (EachLine.StartsWith(TEXT("GameDefaultMap")))
+						{
+							EachLine.Split("=", nullptr, &SpecificGameDefaultMap);
+							FullGameDefaultMapPath = (StarterContentContentDir / SpecificEditorStartupMap) + FPackageName::GetMapPackageExtension();
+							FullGameDefaultMapPath = FullGameDefaultMapPath.Replace(TEXT("Game/"), TEXT(""));
+						}
+					}					
+				}
+			}
 
+			// Look for maps in the content folder. If we don't specify maps for EditorStartup and GameDefault we will use the first we find in here
+			IFileManager::Get().FindFilesRecursive(StarterContentMapFiles, *FPaths::StarterContentDir(), *FileWildcard, /*Files=*/true, /*Directories=*/false);			
+			FString MapPackagePath;
+			if (StarterContentMapFiles.Num() > 0)
+			{
 				const FString BaseMapFilename = FPaths::GetBaseFilename(StarterContentMapFiles[0]);
 
 				FString MapPathRelToContent = FPaths::GetPath(StarterContentMapFiles[0]);
 				FPaths::MakePathRelativeTo(MapPathRelToContent, *StarterContentContentDir);
 
-				const FString MapPackagePath = FString(TEXT("/Game/")) + MapPathRelToContent + TEXT("/") + BaseMapFilename;
-				FileContents += TEXT("[/Script/EngineSettings.GameMapsSettings]") LINE_TERMINATOR;
-				FileContents += FString::Printf(TEXT("EditorStartupMap=%s") LINE_TERMINATOR, *MapPackagePath);
-				FileContents += FString::Printf(TEXT("GameDefaultMap=%s") LINE_TERMINATOR, *MapPackagePath);
-				if (InProjectInfo.bShouldGenerateCode)
-				{
-					FileContents += FString::Printf(TEXT("GlobalDefaultGameMode=\"/Script/%s.%sGameMode\"") LINE_TERMINATOR, *NewProjectName, *NewProjectName);
-				}
+				MapPackagePath = FString(TEXT("/Game/")) + MapPathRelToContent + TEXT("/") + BaseMapFilename;
+			}
+
+			// if either the files we specified don't exist or we didn't specify any, use the first map file we found in the content folder.
+			if (SpecificEditorStartupMap.IsEmpty() || FPaths::FileExists(FullEditorStartupMapPath) == false)
+			{
+				SpecificEditorStartupMap = MapPackagePath;
 			}
+			if (SpecificGameDefaultMap.IsEmpty() || FPaths::FileExists(FullGameDefaultMapPath) == false)
+			{
+				SpecificGameDefaultMap = MapPackagePath;
+			}
+			
+			// Write out the settings for startup map and game default map
+			FileContents += TEXT("[/Script/EngineSettings.GameMapsSettings]") LINE_TERMINATOR;
+			FileContents += FString::Printf(TEXT("EditorStartupMap=%s") LINE_TERMINATOR, *SpecificEditorStartupMap);
+			FileContents += FString::Printf(TEXT("GameDefaultMap=%s") LINE_TERMINATOR, *SpecificGameDefaultMap);
+			if (InProjectInfo.bShouldGenerateCode)
+			{
+				FileContents += FString::Printf(TEXT("GlobalDefaultGameMode=\"/Script/%s.%sGameMode\"") LINE_TERMINATOR, *NewProjectName, *NewProjectName);
+			}			
 		}
 
 		if (WriteOutputFile(DefaultEngineIniFilename, FileContents, OutFailReason))
