Added hardware targeting options and starter content to project creation analytics event

diff --git a/Engine/Source/Editor/GameProjectGeneration/Private/GameProjectUtils.cpp b/Engine/Source/Editor/GameProjectGeneration/Private/GameProjectUtils.cpp
--- a/Engine/Source/Editor/GameProjectGeneration/Private/GameProjectUtils.cpp
+++ b/Engine/Source/Editor/GameProjectGeneration/Private/GameProjectUtils.cpp
@@ -581,6 +581,12 @@ bool GameProjectUtils::CreateProject(const FProjectInformation& InProjectInfo, F
 		EventAttributes.Add(FAnalyticsEventAttribute(TEXT("ProjectType"), InProjectInfo.bShouldGenerateCode ? TEXT("C++ Code") : TEXT("Content Only")));
 		EventAttributes.Add(FAnalyticsEventAttribute(TEXT("Outcome"), bProjectCreationSuccessful ? TEXT("Successful") : TEXT("Failed")));
 
+		UEnum* Enum = FindObject<UEnum>(ANY_PACKAGE, TEXT("EHardwareClass"), true);
+		EventAttributes.Add(FAnalyticsEventAttribute(TEXT("HardwareClass"), Enum ? Enum->GetEnumName(InProjectInfo.TargetedHardware) : FString()));
+		Enum = FindObject<UEnum>(ANY_PACKAGE, TEXT("EGraphicsPreset"), true);
+		EventAttributes.Add(FAnalyticsEventAttribute(TEXT("GraphicsPreset"), Enum ? Enum->GetEnumName(InProjectInfo.DefaultGraphicsPerformance) : FString()));
+		EventAttributes.Add(FAnalyticsEventAttribute(TEXT("StarterContent"), InProjectInfo.bCopyStarterContent ? TEXT("Yes") : TEXT("No")));
+
 		FEngineAnalytics::GetProvider().RecordEvent( TEXT( "Editor.NewProject.ProjectCreated" ), EventAttributes );
 	}
 
