Editor: Fix hardware target settings being written into the wrong config file when creating a new project (and thus not taking effect)
- HW target settings now write to DefaultEngine.ini instead of DefaultEditor.ini when generating a project
- HW target settings now write their enum value as a string instead of an integer when generating a project
#Integrate change 2326075 from 4.5 to main

diff --git a/Engine/Source/Editor/GameProjectGeneration/Private/GameProjectUtils.cpp b/Engine/Source/Editor/GameProjectGeneration/Private/GameProjectUtils.cpp
--- a/Engine/Source/Editor/GameProjectGeneration/Private/GameProjectUtils.cpp
+++ b/Engine/Source/Editor/GameProjectGeneration/Private/GameProjectUtils.cpp
@@ -1126,16 +1126,16 @@ bool GameProjectUtils::CreateProjectFromTemplate(const FProjectInformation& InPr
 
 	// Write out the hardware class target settings chosen for this project
 	{
-		const FString DefaultEditorIniFilename = ProjectConfigPath / TEXT("DefaultEditor.ini");
+		const FString DefaultEngineIniFilename = ProjectConfigPath / TEXT("DefaultEngine.ini");
 
 		FString FileContents;
 		// Load the existing file - if it doesn't exist we create it
-		FFileHelper::LoadFileToString(FileContents, *DefaultEditorIniFilename);
+		FFileHelper::LoadFileToString(FileContents, *DefaultEngineIniFilename);
 
 		FileContents += LINE_TERMINATOR;
 		FileContents += GetHardwareConfigString(InProjectInfo);
 
-		if ( !WriteOutputFile(DefaultEditorIniFilename, FileContents, OutFailReason) )
+		if ( !WriteOutputFile(DefaultEngineIniFilename, FileContents, OutFailReason) )
 		{
 			return false;
 		}
@@ -1467,9 +1467,15 @@ FString GameProjectUtils::GetHardwareConfigString(const FProjectInformation& InP
 {
 	FString HardwareTargeting;
 	
+	FString TargetHardwareAsString;
+	UEnum::GetValueAsString(TEXT("/Script/HardwareTargeting.HardwareTargetingSettings.EHardwareClass"), InProjectInfo.TargetedHardware, /*out*/ TargetHardwareAsString);
+
+	FString GraphicsPresetAsString;
+	UEnum::GetValueAsString(TEXT("/Script/HardwareTargeting.HardwareTargetingSettings.EGraphicsPreset"), InProjectInfo.DefaultGraphicsPerformance, /*out*/ GraphicsPresetAsString);
+
 	HardwareTargeting += TEXT("[/Script/HardwareTargeting.HardwareTargetingSettings]") LINE_TERMINATOR;
-	HardwareTargeting += FString::Printf(TEXT("TargetedHardwareClass=%d") LINE_TERMINATOR, int32(InProjectInfo.TargetedHardware));
-	HardwareTargeting += FString::Printf(TEXT("DefaultGraphicsPerformance=%d") LINE_TERMINATOR, int32(InProjectInfo.DefaultGraphicsPerformance));
+	HardwareTargeting += FString::Printf(TEXT("TargetedHardwareClass=%s") LINE_TERMINATOR, *TargetHardwareAsString);
+	HardwareTargeting += FString::Printf(TEXT("DefaultGraphicsPerformance=%s") LINE_TERMINATOR, *GraphicsPresetAsString);
 	HardwareTargeting += LINE_TERMINATOR;
 
 	return HardwareTargeting;
@@ -1490,6 +1496,9 @@ bool GameProjectUtils::GenerateConfigFiles(const FProjectInformation& InProjectI
 		FileContents += TEXT("[URL]") LINE_TERMINATOR;
 		FileContents += FString::Printf(TEXT("GameName=%s") LINE_TERMINATOR, *NewProjectName);
 		FileContents += LINE_TERMINATOR;
+
+		FileContents += GetHardwareConfigString(InProjectInfo);
+		FileContents += LINE_TERMINATOR;
 		
 		if (InProjectInfo.bCopyStarterContent)
 		{
@@ -1583,8 +1592,6 @@ bool GameProjectUtils::GenerateConfigFiles(const FProjectInformation& InProjectI
 		FileContents += TEXT("bReplaceBlueprintWithClass=true") LINE_TERMINATOR;
 		FileContents += TEXT("bDontLoadBlueprintOutsideEditor=true") LINE_TERMINATOR;
 		FileContents += TEXT("bBlueprintIsNotBlueprintType=true") LINE_TERMINATOR;
-		FileContents += LINE_TERMINATOR;
-		FileContents += GetHardwareConfigString(InProjectInfo);
 
 		if (WriteOutputFile(DefaultEditorIniFilename, FileContents, OutFailReason))
 		{
