Expose ProjectHasCodeFiles, GenerateBasicSourceCode,  GetCurrentProjectModules and DetermineModuleIncludePath from GameProjectGenerationModule

diff --git a/Engine/Source/Editor/GameProjectGeneration/Private/GameProjectGenerationModule.cpp b/Engine/Source/Editor/GameProjectGeneration/Private/GameProjectGenerationModule.cpp
--- a/Engine/Source/Editor/GameProjectGeneration/Private/GameProjectGenerationModule.cpp
+++ b/Engine/Source/Editor/GameProjectGeneration/Private/GameProjectGenerationModule.cpp
@@ -86,10 +86,24 @@ bool FGameProjectGenerationModule::UpdateCodeProject(FText& OutFailReason)
 	return GameProjectUtils::GenerateCodeProjectFiles(FPaths::GetProjectFilePath(), OutFailReason);
 }
 
+bool FGameProjectGenerationModule::GenerateBasicSourceCode(TArray<FString>& OutCreatedFiles, FText& OutFailReason)
+{
+	return GameProjectUtils::GenerateBasicSourceCode(OutCreatedFiles, OutFailReason);
+}
+
+bool FGameProjectGenerationModule::ProjectHasCodeFiles()
+{
+	return GameProjectUtils::ProjectHasCodeFiles();
+}
+
+FString FGameProjectGenerationModule::DetermineModuleIncludePath(const FModuleContextInfo& ModuleInfo, const FString& FileRelativeTo)
+{
+	return GameProjectUtils::DetermineModuleIncludePath(ModuleInfo, FileRelativeTo);
+}
 
-int32 FGameProjectGenerationModule::GetProjectCodeFileCount()
+TArray<FModuleContextInfo> FGameProjectGenerationModule::GetCurrentProjectModules()
 {
-	return GameProjectUtils::GetProjectCodeFileCount();
+	return GameProjectUtils::GetCurrentProjectModules();
 }
 
 void FGameProjectGenerationModule::GetProjectSourceDirectoryInfo(int32& OutNumFiles, int64& OutDirectorySize)
diff --git a/Engine/Source/Editor/GameProjectGeneration/Private/GameProjectUtils.cpp b/Engine/Source/Editor/GameProjectGeneration/Private/GameProjectUtils.cpp
--- a/Engine/Source/Editor/GameProjectGeneration/Private/GameProjectUtils.cpp
+++ b/Engine/Source/Editor/GameProjectGeneration/Private/GameProjectUtils.cpp
@@ -1606,6 +1606,18 @@
 	return true;
 }
 
+bool GameProjectUtils::GenerateBasicSourceCode(TArray<FString>& OutCreatedFiles, FText& OutFailReason)
+{
+	TArray<FString> StartupModuleNames;
+	if (GameProjectUtils::GenerateBasicSourceCode(FPaths::GameSourceDir().LeftChop(1), FApp::GetGameName(), FPaths::GameDir(), StartupModuleNames, OutCreatedFiles, OutFailReason))
+	{
+		GameProjectUtils::UpdateProject(&StartupModuleNames);
+		return true;
+	}
+
+	return false;
+}
+
 bool GameProjectUtils::GenerateBasicSourceCode(const FString& NewProjectSourcePath, const FString& NewProjectName, const FString& NewProjectRoot, TArray<FString>& OutGeneratedStartupModuleNames, TArray<FString>& OutCreatedFiles, FText& OutFailReason)
 {
 	const FString GameModulePath = NewProjectSourcePath / NewProjectName;
@@ -1774,7 +1786,7 @@
 	return (StarterContentFiles.Num() > 0);
 }
 
-TArray<GameProjectUtils::FModuleContextInfo> GameProjectUtils::GetCurrentProjectModules()
+TArray<FModuleContextInfo> GameProjectUtils::GetCurrentProjectModules()
 {
 	const FProjectDescriptor* const CurrentProject = IProjectManager::Get().GetCurrentProject();
 	check(CurrentProject);
@@ -2180,6 +2192,58 @@
 	return ReturnString;
 }
 
+FString GameProjectUtils::DetermineModuleIncludePath(const FModuleContextInfo& ModuleInfo, const FString& FileRelativeTo)
+{
+	FString ModuleIncludePath;
+
+	if(FindSourceFileInProject(ModuleInfo.ModuleName + ".h", ModuleInfo.ModuleSourcePath, ModuleIncludePath))
+	{
+		// Work out where the module header is; 
+		// if it's Public then we can include it without any path since all Public and Classes folders are on the include path
+		// if it's located elsewhere, then we'll need to include it relative to the module source root as we can't guarantee 
+		// that other folders are on the include paths
+		EClassLocation ModuleLocation;
+		if(GetClassLocation(ModuleIncludePath, ModuleInfo, ModuleLocation))
+		{
+			if(ModuleLocation == EClassLocation::Public || ModuleLocation == EClassLocation::Classes)
+			{
+				ModuleIncludePath = ModuleInfo.ModuleName + ".h";
+			}
+			else
+			{
+				// If the path to our new class is the same as the path to the module, we can include it directly
+				const FString ModulePath = FPaths::ConvertRelativePathToFull(FPaths::GetPath(ModuleIncludePath));
+				const FString ClassPath = FPaths::ConvertRelativePathToFull(FPaths::GetPath(FileRelativeTo));
+				if(ModulePath == ClassPath)
+				{
+					ModuleIncludePath = ModuleInfo.ModuleName + ".h";
+				}
+				else
+				{
+					// Updates ModuleIncludePath internally
+					if(!FPaths::MakePathRelativeTo(ModuleIncludePath, *ModuleInfo.ModuleSourcePath))
+					{
+						// Failed; just assume we can include it without any relative path
+						ModuleIncludePath = ModuleInfo.ModuleName + ".h";
+					}
+				}
+			}
+		}
+		else
+		{
+			// Failed; just assume we can include it without any relative path
+			ModuleIncludePath = ModuleInfo.ModuleName + ".h";
+		}
+	}
+	else
+	{
+		// This could potentially fail when generating new projects if the module file hasn't yet been created; just assume we can include it without any relative path
+		ModuleIncludePath = ModuleInfo.ModuleName + ".h";
+	}
+
+	return ModuleIncludePath;
+}
+
 /**
  * Generates UObject class constructor definition with property overrides.
  *
@@ -2348,51 +2412,8 @@
 	}
 
 	// Calculate the correct include path for the module header
-	FString ModuleIncludePath;
-	if(FindSourceFileInProject(ModuleInfo.ModuleName + ".h", ModuleInfo.ModuleSourcePath, ModuleIncludePath))
-	{
-		// Work out where the module header is; 
-		// if it's Public then we can include it without any path since all Public and Classes folders are on the include path
-		// if it's located elsewhere, then we'll need to include it relative to the module source root as we can't guarantee 
-		// that other folders are on the include paths
-		EClassLocation ModuleLocation;
-		if(GetClassLocation(ModuleIncludePath, ModuleInfo, ModuleLocation))
-		{
-			if(ModuleLocation == EClassLocation::Public || ModuleLocation == EClassLocation::Classes)
-			{
-				ModuleIncludePath = ModuleInfo.ModuleName + ".h";
-			}
-			else
-			{
-				// If the path to our new class is the same as the path to the module, we can include it directly
-				const FString ModulePath = FPaths::ConvertRelativePathToFull(FPaths::GetPath(ModuleIncludePath));
-				const FString ClassPath = FPaths::ConvertRelativePathToFull(FPaths::GetPath(NewCPPFileName));
-				if(ModulePath == ClassPath)
-				{
-					ModuleIncludePath = ModuleInfo.ModuleName + ".h";
-				}
-				else
-				{
-					// Updates ModuleIncludePath internally
-					if(!FPaths::MakePathRelativeTo(ModuleIncludePath, *ModuleInfo.ModuleSourcePath))
-					{
-						// Failed; just assume we can include it without any relative path
-						ModuleIncludePath = ModuleInfo.ModuleName + ".h";
-					}
-				}
-			}
-		}
-		else
-		{
-			// Failed; just assume we can include it without any relative path
-			ModuleIncludePath = ModuleInfo.ModuleName + ".h";
-		}
-	}
-	else
-	{
-		// This could potentially fail when generating new projects if the module file hasn't yet been created; just assume we can include it without any relative path
-		ModuleIncludePath = ModuleInfo.ModuleName + ".h";
-	}
+	const FString ModuleIncludePath = DetermineModuleIncludePath(ModuleInfo, NewCPPFileName);
 
+
 	FString EventualConstructorDefinition;
 	if (PropertyOverrides.Num() != 0)
diff --git a/Engine/Source/Editor/GameProjectGeneration/Private/GameProjectUtils.h b/Engine/Source/Editor/GameProjectGeneration/Private/GameProjectUtils.h
--- a/Engine/Source/Editor/GameProjectGeneration/Private/GameProjectUtils.h
+++ b/Engine/Source/Editor/GameProjectGeneration/Private/GameProjectUtils.h
@@ -4,6 +4,8 @@
 
 #include "HardwareTargetingModule.h"
 
+struct FModuleContextInfo;
+
 struct FProjectInformation
 {
 	FProjectInformation(FString InProjectFilename, bool bInGenerateCode, bool bInCopyStarterContent, FString InTemplateFile = FString())
@@ -28,19 +30,6 @@ struct FProjectInformation
 class GameProjectUtils
 {
 public:
-	/** Context information used when validating that source code is being placed in the correct place for a given module */
-	struct FModuleContextInfo
-	{
-		/** Path to the Source folder of the module */
-		FString ModuleSourcePath;
-
-		/** Name of the module */
-		FString ModuleName;
-
-		/** Type of this module, eg, Runtime, Editor, etc */
-		EHostType::Type ModuleType;
-	};
-
 	/** Where is this class located within the Source folder? */
 	enum class EClassLocation : uint8
 	{
@@ -273,6 +262,15 @@ public:
 	/** Clear the list of supported target platforms */
 	static void ClearSupportedTargetPlatforms();
 
+	/** Returns the path to the module's include header */
+	static FString DetermineModuleIncludePath(const FModuleContextInfo& ModuleInfo, const FString& FileRelativeTo);
+
+	/** Creates the basic source code for a new project. On failure, OutFailReason will be populated. */
+	static bool GenerateBasicSourceCode(TArray<FString>& OutCreatedFiles, FText& OutFailReason);
+
+	/** Returns true if the currently loaded project has code files */
+	static bool ProjectHasCodeFiles();
+
 	/** Returns the contents of the specified template file */
 	static bool ReadTemplateFile(const FString& TemplateFileName, FString& OutFileContents, FText& OutFailReason);
 private:
@@ -400,9 +398,6 @@ private:
 	/** Checks the specified game project file out from source control */
 	static bool CheckoutGameProjectFile(const FString& ProjectFilename, FText& OutFailReason);
 
-	/** Returns true if the currently loaded project has code files */
-	static bool ProjectHasCodeFiles();
-
 	/** Internal handler for AddCodeToProject*/
 	static bool AddCodeToProject_Internal(const FString& NewClassName, const FString& NewClassPath, const FModuleContextInfo& ModuleInfo, const FNewClassInfo ParentClassInfo, FString& OutHeaderFilePath, FString& OutCppFilePath, FText& OutFailReason);
 
diff --git a/Engine/Source/Editor/GameProjectGeneration/Private/SNewClassDialog.cpp b/Engine/Source/Editor/GameProjectGeneration/Private/SNewClassDialog.cpp
--- a/Engine/Source/Editor/GameProjectGeneration/Private/SNewClassDialog.cpp
+++ b/Engine/Source/Editor/GameProjectGeneration/Private/SNewClassDialog.cpp
@@ -27,7 +27,7 @@ struct FParentClassItem
 class FNativeClassParentFilter : public IClassViewerFilter
 {
 public:
-	FNativeClassParentFilter(const TSharedPtr<GameProjectUtils::FModuleContextInfo>* InSelectedModuleInfoPtr)
+	FNativeClassParentFilter(const TSharedPtr<FModuleContextInfo>* InSelectedModuleInfoPtr)
 		: SelectedModuleInfoPtr(InSelectedModuleInfoPtr)
 	{
 	}
@@ -42,7 +42,7 @@ public:
 
 		// Is this class in the same module as our current module?
 		const FString ClassModuleName = InClass->GetOutermost()->GetName().RightChop( FString(TEXT("/Script/")).Len() );
-		const TSharedPtr<GameProjectUtils::FModuleContextInfo>& ModuleInfo = *SelectedModuleInfoPtr;
+		const TSharedPtr<FModuleContextInfo>& ModuleInfo = *SelectedModuleInfoPtr;
 		const bool bIsInDestinationModule = (ModuleInfo.IsValid() && ModuleInfo->ModuleName == ClassModuleName);
 
 		// You need API if you are either not UObject itself and you are not in the destination module
@@ -65,20 +65,20 @@ public:
 
 private:
 	/** Pointer to the currently selected module in the new class dialog */
-	const TSharedPtr<GameProjectUtils::FModuleContextInfo>* SelectedModuleInfoPtr;
+	const TSharedPtr<FModuleContextInfo>* SelectedModuleInfoPtr;
 };
 
 BEGIN_SLATE_FUNCTION_BUILD_OPTIMIZATION
 void SNewClassDialog::Construct( const FArguments& InArgs )
 {
 	{
-		TArray<GameProjectUtils::FModuleContextInfo> CurrentModules = GameProjectUtils::GetCurrentProjectModules();
+		TArray<FModuleContextInfo> CurrentModules = GameProjectUtils::GetCurrentProjectModules();
 		check(CurrentModules.Num()); // this should never happen since GetCurrentProjectModules is supposed to add a dummy runtime module if the project currently has no modules
 
 		AvailableModules.Reserve(CurrentModules.Num());
-		for(const GameProjectUtils::FModuleContextInfo& ModuleInfo : CurrentModules)
+		for(const FModuleContextInfo& ModuleInfo : CurrentModules)
 		{
-			AvailableModules.Emplace(MakeShareable(new GameProjectUtils::FModuleContextInfo(ModuleInfo)));
+			AvailableModules.Emplace(MakeShareable(new FModuleContextInfo(ModuleInfo)));
 		}
 	}
 
@@ -445,7 +445,7 @@ void SNewClassDialog::Construct( const FArguments& InArgs )
 											.AutoWidth()
 											.Padding(6.0f, 0.0f, 0.0f, 0.0f)
 											[
-												SAssignNew(AvailableModulesCombo, SComboBox<TSharedPtr<GameProjectUtils::FModuleContextInfo>>)
+												SAssignNew(AvailableModulesCombo, SComboBox<TSharedPtr<FModuleContextInfo>>)
 												.ToolTipText( LOCTEXT("ModuleComboToolTip", "Choose the target module for your new class") )
 												.OptionsSource( &AvailableModules )
 												.InitiallySelectedItem( SelectedModuleInfo )
@@ -1058,7 +1058,7 @@ FText SNewClassDialog::GetSelectedModuleComboText() const
 	return FText::Format(LOCTEXT("ModuleComboEntry", "{ModuleName} ({ModuleType})"), Args);
 }
 
-void SNewClassDialog::SelectedModuleComboBoxSelectionChanged(TSharedPtr<GameProjectUtils::FModuleContextInfo> Value, ESelectInfo::Type SelectInfo)
+void SNewClassDialog::SelectedModuleComboBoxSelectionChanged(TSharedPtr<FModuleContextInfo> Value, ESelectInfo::Type SelectInfo)
 {
 	const FString& OldModulePath = SelectedModuleInfo->ModuleSourcePath;
 	const FString& NewModulePath = Value->ModuleSourcePath;
@@ -1075,7 +1075,7 @@ void SNewClassDialog::SelectedModuleComboBoxSelectionChanged(TSharedPtr<GameProj
 	UpdateInputValidity();
 }
 
-TSharedRef<SWidget> SNewClassDialog::MakeWidgetForSelectedModuleCombo(TSharedPtr<GameProjectUtils::FModuleContextInfo> Value)
+TSharedRef<SWidget> SNewClassDialog::MakeWidgetForSelectedModuleCombo(TSharedPtr<FModuleContextInfo> Value)
 {
 	FFormatNamedArguments Args;
 	Args.Add(TEXT("ModuleName"), FText::FromString(Value->ModuleName));
diff --git a/Engine/Source/Editor/GameProjectGeneration/Private/SNewClassDialog.h b/Engine/Source/Editor/GameProjectGeneration/Private/SNewClassDialog.h
--- a/Engine/Source/Editor/GameProjectGeneration/Private/SNewClassDialog.h
+++ b/Engine/Source/Editor/GameProjectGeneration/Private/SNewClassDialog.h
@@ -126,10 +126,10 @@ private:
 	FText GetSelectedModuleComboText() const;
 
 	/** Called when the currently selected module is changed */
-	void SelectedModuleComboBoxSelectionChanged(TSharedPtr<GameProjectUtils::FModuleContextInfo> Value, ESelectInfo::Type SelectInfo);
+	void SelectedModuleComboBoxSelectionChanged(TSharedPtr<FModuleContextInfo> Value, ESelectInfo::Type SelectInfo);
 
 	/** Create the widget to use as the combo box entry for the given module info */
-	TSharedRef<SWidget> MakeWidgetForSelectedModuleCombo(TSharedPtr<GameProjectUtils::FModuleContextInfo> Value);
+	TSharedRef<SWidget> MakeWidgetForSelectedModuleCombo(TSharedPtr<FModuleContextInfo> Value);
 
 private:
 
@@ -176,7 +176,7 @@ private:
 	TSharedPtr<SEditableTextBox> ClassNameEditBox;
 
 	/** The available modules combo box */
-	TSharedPtr<SComboBox<TSharedPtr<GameProjectUtils::FModuleContextInfo>>> AvailableModulesCombo;
+	TSharedPtr<SComboBox<TSharedPtr<FModuleContextInfo>>> AvailableModulesCombo;
 
 	/** The name of the class being created */
 	FString NewClassName;
@@ -218,8 +218,8 @@ private:
 	GameProjectUtils::EClassLocation ClassLocation;
 
 	/** Information about the currently available modules for this project */
-	TArray<TSharedPtr<GameProjectUtils::FModuleContextInfo>> AvailableModules;
+	TArray<TSharedPtr<FModuleContextInfo>> AvailableModules;
 
 	/** Information about the currently selected module; used for class validation */
-	TSharedPtr<GameProjectUtils::FModuleContextInfo> SelectedModuleInfo;
+	TSharedPtr<FModuleContextInfo> SelectedModuleInfo;
 };
diff --git a/Engine/Source/Editor/GameProjectGeneration/Public/GameProjectGenerationModule.h b/Engine/Source/Editor/GameProjectGeneration/Public/GameProjectGenerationModule.h
--- a/Engine/Source/Editor/GameProjectGeneration/Public/GameProjectGenerationModule.h
+++ b/Engine/Source/Editor/GameProjectGeneration/Public/GameProjectGenerationModule.h
@@ -7,6 +7,20 @@
 
 struct FTemplateCategory;
 
+/** Context information used when validating that source code is being placed in the correct place for a given module */
+struct FModuleContextInfo
+{
+	/** Path to the Source folder of the module */
+	FString ModuleSourcePath;
+
+	/** Name of the module */
+	FString ModuleName;
+
+	/** Type of this module, eg, Runtime, Editor, etc */
+	EHostType::Type ModuleType;
+};
+
+
 /**
  * Game Project Generation module
  */
@@ -63,7 +77,12 @@ public:
 	virtual bool UpdateCodeProject(FText& OutFailReason);
 
 	/** Gets the current projects source file count */
-	virtual int32 GetProjectCodeFileCount();
+	virtual bool ProjectHasCodeFiles();
+
+	/** Returns the path to the module's include header */
+	virtual FString DetermineModuleIncludePath(const FModuleContextInfo& ModuleInfo, const FString& FileRelativeTo);
+
+	virtual TArray<FModuleContextInfo> GetCurrentProjectModules();
 
 	/** Gets file and size info about the source directory */
 	virtual void GetProjectSourceDirectoryInfo(int32& OutNumFiles, int64& OutDirectorySize);
@@ -71,6 +90,9 @@ public:
 	/** Warn the user if the project filename is invalid in case they renamed it outside the editor */
 	virtual void CheckAndWarnProjectFilenameValid();
 
+	/** Generate basic project source code */
+	virtual bool GenerateBasicSourceCode(TArray<FString>& OutCreatedFiles, FText& OutFailReason);
+
 	/**
 	 * Update the list of supported target platforms based upon the parameters provided
 	 * This will take care of checking out and saving the updated .uproject file automatically
diff --git a/Engine/Source/Editor/MainFrame/Private/Frame/MainFrameActions.cpp b/Engine/Source/Editor/MainFrame/Private/Frame/MainFrameActions.cpp
--- a/Engine/Source/Editor/MainFrame/Private/Frame/MainFrameActions.cpp
+++ b/Engine/Source/Editor/MainFrame/Private/Frame/MainFrameActions.cpp
@@ -506,7 +506,7 @@ void FMainFrameActionCallbacks::PackageProject( const FName InPlatformInfoName )
 {
 	// does the project have any code?
 	FGameProjectGenerationModule& GameProjectModule = FModuleManager::LoadModuleChecked<FGameProjectGenerationModule>(TEXT("GameProjectGeneration"));
-	bool bProjectHasCode = GameProjectModule.Get().GetProjectCodeFileCount() > 0;
+	bool bProjectHasCode = GameProjectModule.Get().ProjectHasCodeFiles();
 
 	const PlatformInfo::FPlatformInfo* const PlatformInfo = PlatformInfo::FindPlatformInfo(InPlatformInfoName);
 	check(PlatformInfo);
@@ -1101,7 +1101,7 @@ void FMainFrameActionCallbacks::CreateUatTask( const FString& CommandLine, const
 
 	FString UatPath = FPaths::ConvertRelativePathToFull(FPaths::EngineDir() / TEXT("Build/BatchFiles") / RunUATScriptName);
 	FGameProjectGenerationModule& GameProjectModule = FModuleManager::LoadModuleChecked<FGameProjectGenerationModule>(TEXT("GameProjectGeneration"));
-	bool bHasCode = GameProjectModule.Get().GetProjectCodeFileCount() > 0;
+	bool bHasCode = GameProjectModule.Get().ProjectHasCodeFiles();
 
 	if (!FPaths::FileExists(UatPath))
 	{
diff --git a/Engine/Source/Editor/UnrealEd/Private/Kismet2/DebuggerCommands.cpp b/Engine/Source/Editor/UnrealEd/Private/Kismet2/DebuggerCommands.cpp
--- a/Engine/Source/Editor/UnrealEd/Private/Kismet2/DebuggerCommands.cpp
+++ b/Engine/Source/Editor/UnrealEd/Private/Kismet2/DebuggerCommands.cpp
@@ -1394,7 +1394,7 @@ bool FInternalPlayWorldCommandCallbacks::IsReadyToLaunchOnDevice(FString DeviceI
 	check(PlatformInfo);
 
 	FGameProjectGenerationModule& GameProjectModule = FModuleManager::LoadModuleChecked<FGameProjectGenerationModule>(TEXT("GameProjectGeneration"));
-	bool bHasCode = GameProjectModule.Get().GetProjectCodeFileCount() > 0;
+	bool bHasCode = GameProjectModule.Get().ProjectHasCodeFiles();
 
 	if (PlatformInfo->SDKStatus == PlatformInfo::EPlatformSDKStatus::NotInstalled)
 	{
diff --git a/Engine/Source/Editor/UnrealEd/Private/PlayLevel.cpp b/Engine/Source/Editor/UnrealEd/Private/PlayLevel.cpp
--- a/Engine/Source/Editor/UnrealEd/Private/PlayLevel.cpp
+++ b/Engine/Source/Editor/UnrealEd/Private/PlayLevel.cpp
@@ -1267,7 +1267,7 @@ void UEditorEngine::PlayUsingLauncher()
 
 		// does the project have any code?
 		FGameProjectGenerationModule& GameProjectModule = FModuleManager::LoadModuleChecked<FGameProjectGenerationModule>(TEXT("GameProjectGeneration"));
-		bool bHasCode = GameProjectModule.Get().GetProjectCodeFileCount() > 0;
+		bool bHasCode = GameProjectModule.Get().ProjectHasCodeFiles();
 
 		// Setup launch profile, keep the setting here to a minimum.
 		ILauncherProfileRef LauncherProfile = LauncherServicesModule.CreateProfile(TEXT("Play On Device"));
