Potential CIS fix for mac compile error

diff --git a/Engine/Source/Editor/GameProjectGeneration/Public/GameProjectGenerationModule.h b/Engine/Source/Editor/GameProjectGeneration/Public/GameProjectGenerationModule.h
--- a/Engine/Source/Editor/GameProjectGeneration/Public/GameProjectGenerationModule.h
+++ b/Engine/Source/Editor/GameProjectGeneration/Public/GameProjectGenerationModule.h
@@ -17,7 +17,7 @@ struct FModuleContextInfo
 	FString ModuleName;
 
 	/** Type of this module, eg, Runtime, Editor, etc */
-	EHostType::Type ModuleType;
+	enum EHostType::Type ModuleType;
 };
 
 
