Fix non-unity

diff --git a/Engine/Source/Editor/GameProjectGeneration/Private/GameProjectGenerationPrivatePCH.h b/Engine/Source/Editor/GameProjectGeneration/Private/GameProjectGenerationPrivatePCH.h
--- a/Engine/Source/Editor/GameProjectGeneration/Private/GameProjectGenerationPrivatePCH.h
+++ b/Engine/Source/Editor/GameProjectGeneration/Private/GameProjectGenerationPrivatePCH.h
@@ -10,6 +10,7 @@
 DECLARE_LOG_CATEGORY_EXTERN(LogGameProjectGeneration, Verbose, All);
 
 #include "GameProjectGenerationClasses.h"
+#include "GameProjectGenerationModule.h"
 #include "GameProjectUtils.h"
 #include "SProjectBrowser.h"
 #include "SNewProjectWizard.h"
