fix for Mac compile
#ue4

diff --git a/Engine/Source/Editor/GameProjectGeneration/Public/GameProjectGenerationModule.h b/Engine/Source/Editor/GameProjectGeneration/Public/GameProjectGenerationModule.h
--- a/Engine/Source/Editor/GameProjectGeneration/Public/GameProjectGenerationModule.h
+++ b/Engine/Source/Editor/GameProjectGeneration/Public/GameProjectGenerationModule.h
@@ -4,6 +4,7 @@
 #pragma once
 
 #include "ModuleInterface.h"
+#include "ModuleDescriptor.h"
 
 struct FTemplateCategory;
 
