#UE4 DOCS: Added tutorial highlighting meta data to the Add Code to Project wizard: the major frame, the "type of code to add" scroll box, the class name, and path.

diff --git a/Engine/Source/Editor/GameProjectGeneration/Private/SNewClassDialog.cpp b/Engine/Source/Editor/GameProjectGeneration/Private/SNewClassDialog.cpp
--- a/Engine/Source/Editor/GameProjectGeneration/Private/SNewClassDialog.cpp
+++ b/Engine/Source/Editor/GameProjectGeneration/Private/SNewClassDialog.cpp
@@ -11,6 +11,7 @@
 #include "EditorClassUtils.h"
 #include "SWizard.h"
 #include "SHyperlink.h"
+#include "TutorialMetaData.h"
 
 
 #define LOCTEXT_NAMESPACE "GameProjectGeneration"
@@ -151,6 +152,7 @@ void SNewClassDialog::Construct( const FArguments& InArgs )
 		.BorderImage( FEditorStyle::GetBrush("Docking.Tab.ContentAreaBrush") )
 		[
 			SNew(SVerticalBox)
+			.AddMetaData<FTutorialMetaData>(TEXT("AddCodeMajorAnchor"))
 
 			+SVerticalBox::Slot()
 			[
@@ -223,6 +225,7 @@ void SNewClassDialog::Construct( const FArguments& InArgs )
 					.Padding(0, 10)
 					[
 						SNew(SBorder)
+						.AddMetaData<FTutorialMetaData>(TEXT("AddCodeOptions"))
 						.BorderImage( FEditorStyle::GetBrush("ToolPanel.GroupBorder") )
 						[
 							SNew(SVerticalBox)
@@ -258,7 +261,7 @@ void SNewClassDialog::Construct( const FArguments& InArgs )
 					.AutoHeight()
 					[
 						SNew(SHorizontalBox)
-
+						
 						// Class label
 						+SHorizontalBox::Slot()
 						.AutoWidth()
@@ -431,6 +434,7 @@ void SNewClassDialog::Construct( const FArguments& InArgs )
 									[
 										SNew(SBox)
 										.HeightOverride(EditableTextHeight)
+										.AddMetaData<FTutorialMetaData>(TEXT("ClassName"))
 										[
 											SNew(SHorizontalBox)
 
@@ -528,6 +532,7 @@ void SNewClassDialog::Construct( const FArguments& InArgs )
 									[
 										SNew(SBox)
 										.HeightOverride(EditableTextHeight)
+										.AddMetaData<FTutorialMetaData>(TEXT("Path"))
 										[
 											SNew(SHorizontalBox)
 
