Tweaked new project dialog so it can show class and asset types used by a template.
Revised template descriptions to provide more details.

diff --git a/Engine/Source/Editor/GameProjectGeneration/Classes/TemplateProjectDefs.h b/Engine/Source/Editor/GameProjectGeneration/Classes/TemplateProjectDefs.h
--- a/Engine/Source/Editor/GameProjectGeneration/Classes/TemplateProjectDefs.h
+++ b/Engine/Source/Editor/GameProjectGeneration/Classes/TemplateProjectDefs.h
@@ -90,6 +90,12 @@ class UTemplateProjectDefs : public UObject
 	UPROPERTY(config)
 	FName Category;
 
+	UPROPERTY(config)
+	FString ClassTypes;
+
+	UPROPERTY(config)
+	FString AssetTypes;
+
 	/** Fixes up all strings in this definitions object to replace \%TEMPLATENAME\% with the supplied template name and \%PROJECTNAME\% with the supplied project name */
 	void FixupStrings(const FString& TemplateName, const FString& ProjectName);
 
diff --git a/Engine/Source/Editor/GameProjectGeneration/Private/SNewProjectWizard.cpp b/Engine/Source/Editor/GameProjectGeneration/Private/SNewProjectWizard.cpp
--- a/Engine/Source/Editor/GameProjectGeneration/Private/SNewProjectWizard.cpp
+++ b/Engine/Source/Editor/GameProjectGeneration/Private/SNewProjectWizard.cpp
@@ -42,8 +42,11 @@ struct FTemplateItem
 	TSharedPtr<FSlateBrush> Thumbnail;
 	TSharedPtr<FSlateBrush> PreviewImage;
 
-	FTemplateItem(FText InName, FText InDescription, bool bInGenerateCode, FName InType, FString InSortKey, FString InProjectFile, TSharedPtr<FSlateBrush> InThumbnail, TSharedPtr<FSlateBrush> InPreviewImage)
+	FString		ClassTypes;
+	FString		AssetTypes;
+	FTemplateItem(FText InName, FText InDescription, bool bInGenerateCode, FName InType, FString InSortKey, FString InProjectFile, TSharedPtr<FSlateBrush> InThumbnail, TSharedPtr<FSlateBrush> InPreviewImage,FString InClassTypes, FString InAssetTypes)
 		: Name(InName), Description(InDescription), bGenerateCode(bInGenerateCode), Type(InType), SortKey(MoveTemp(InSortKey)), ProjectFile(MoveTemp(InProjectFile)), Thumbnail(InThumbnail), PreviewImage(InPreviewImage)
+		, ClassTypes(InClassTypes), AssetTypes(InAssetTypes)
 	{}
 };
 
@@ -345,7 +348,7 @@ void SNewProjectWizard::Construct( const FArguments& InArgs )
 
 			// Wizard
 			+SOverlay::Slot()
-			.Padding(UniformPadding / 2)
+			.Padding(UniformPadding / 2.0f)
 			[
 				SAssignNew(MainWizard, SWizard)
 				.ShowPageList(false)
@@ -362,12 +365,12 @@ void SNewProjectWizard::Construct( const FArguments& InArgs )
 				[
 					SNew(SBorder)
 					.BorderImage(FEditorStyle::GetBrush("NoBorder"))
-					.Padding(FMargin(UniformPadding / 2, UniformPadding / 2, UniformPadding / 2, 0))
+					.Padding(FMargin(UniformPadding / 2.0f, UniformPadding / 2.0f, UniformPadding / 2.0f, 0.0f))
 					[
 						SNew(SVerticalBox)
 
 						+ SVerticalBox::Slot()
-						.Padding(FMargin(0, 0, 0, 15))
+						//.Padding(FMargin(0, 0, 0, 15))
 						.AutoHeight()
 						[
 							SNew(STextBlock)
@@ -406,7 +409,7 @@ void SNewProjectWizard::Construct( const FArguments& InArgs )
 									]
 										
 									+ SHorizontalBox::Slot()
-									.Padding(UniformPadding, 0)
+									.Padding(UniformPadding, 0.0f)
 									.AutoWidth()
 									[
 										Separator
@@ -417,7 +420,7 @@ void SNewProjectWizard::Construct( const FArguments& InArgs )
 									[
 										SNew(SScrollBox)
 										+ SScrollBox::Slot()
-										.Padding(UniformPadding, 0)
+										.Padding(UniformPadding, 0.0f)
 										[
 											SNew(SVerticalBox)
 
@@ -425,7 +428,7 @@ void SNewProjectWizard::Construct( const FArguments& InArgs )
 											+ SVerticalBox::Slot()
 											.AutoHeight()
 											.HAlign(HAlign_Center)
-											.Padding(FMargin(0, 0, 0, 15.f))
+											.Padding(FMargin(0.0f, 0.0f, 0.0f, 15.f))
 											[
 												SNew(SBox)
 												.Visibility(this, &SNewProjectWizard::GetSelectedTemplatePreviewVisibility)
@@ -437,7 +440,7 @@ void SNewProjectWizard::Construct( const FArguments& InArgs )
 													+ SOverlay::Slot()
 													[
 														SNew(SBorder)
-														.Padding(FMargin(0, 0, 0, 4.f))
+														.Padding(FMargin(0.0f, 0.0f, 0.0f, 4.f))
 														.BorderImage(FEditorStyle::GetBrush("ContentBrowser.ThumbnailShadow"))
 														[
 															SNew(SImage)
@@ -448,11 +451,11 @@ void SNewProjectWizard::Construct( const FArguments& InArgs )
 													+ SOverlay::Slot()
 													.HAlign(HAlign_Right)
 													.VAlign(VAlign_Top)
-													.Padding(10)
+													.Padding(10.0f)
 													[
 														SNew(SBox)
-														.WidthOverride(48)
-														.HeightOverride(48)
+														.WidthOverride(48.0f)
+														.HeightOverride(48.0f)
 														[
 															SNew(SImage)
 															.Image(this, &SNewProjectWizard::GetSelectedTemplateTypeImage)
@@ -463,7 +466,7 @@ void SNewProjectWizard::Construct( const FArguments& InArgs )
 
 											// Template Name
 											+ SVerticalBox::Slot()
-											.Padding(FMargin(0, 0, 0, 10))
+											.Padding(FMargin(0.0f, 0.0f, 0.0f, 10.0f))
 											.AutoHeight()
 											[
 												SNew(STextBlock)
@@ -479,17 +482,65 @@ void SNewProjectWizard::Construct( const FArguments& InArgs )
 												.AutoWrapText(true)
 												.Text(this, &SNewProjectWizard::GetSelectedTemplateProperty<FText>, &FTemplateItem::Description)
 											]
+											
+											// Class types
+											+SVerticalBox::Slot()
+											.AutoHeight()
+											.Padding(FMargin(0.0f, 5.0f, 0.0f, 5.0f))
+											[
+												SNew(SBox)
+												.Visibility(this, &SNewProjectWizard::GetSelectedTemplateClassVisibility)
+												[
+													SNew(SVerticalBox)
+													+SVerticalBox::Slot()
+													[
+														SNew(STextBlock)
+														.TextStyle(FEditorStyle::Get(), "GameProjectDialog.FeatureText")
+														.Text(LOCTEXT("ProjectTemplateClassTypes", "Class Type References:"))
+													]
+													+ SVerticalBox::Slot()
+													.AutoHeight()
+													[
+														SNew(STextBlock)
+														.AutoWrapText(true)
+														.Text(this, &SNewProjectWizard::GetSelectedTemplateClassTypes)
+													]
+												]
+												
+											]
+											// Asset types
+											+ SVerticalBox::Slot()
+											.AutoHeight()
+											.Padding(FMargin(0.0f, 5.0f, 0.0f, 5.0f))
+											[
+												SNew(SBox)
+												.Visibility(this, &SNewProjectWizard::GetSelectedTemplateAssetVisibility)
+												[
+													SNew(SVerticalBox)
+													+ SVerticalBox::Slot()
+													[
+														SNew(STextBlock)
+														.TextStyle(FEditorStyle::Get(), "GameProjectDialog.FeatureText")
+														.Text(LOCTEXT("ProjectTemplateAssetTypes", "Asset Type References:"))
+													]
+													+ SVerticalBox::Slot()
+													.AutoHeight()
+													[
+														SNew(STextBlock)
+														.AutoWrapText(true)
+														.Text(this, &SNewProjectWizard::GetSelectedTemplateAssetTypes)
+													]
+												]
+											]
 										]
 									]
-
-
 								]
 							]
 						]
 
 						+ SVerticalBox::Slot()
 						.AutoHeight()
-						.Padding(FMargin(0, 15, 0, 0))
+						.Padding(FMargin(0.0f, 15.0f, 0.0f, 0.0f))
 						[
 							SNew(SScrollBox)
 
@@ -800,6 +851,26 @@ TSharedPtr<FTemplateItem> SNewProjectWizard::GetSelectedTemplateItem() const
 	return NULL;
 }
 
+FString SNewProjectWizard::GetSelectedTemplateClassTypes() const
+{
+	return GetSelectedTemplateProperty<FString>(&FTemplateItem::ClassTypes);
+}
+
+EVisibility SNewProjectWizard::GetSelectedTemplateClassVisibility() const
+{
+	return GetSelectedTemplateProperty<FString>(&FTemplateItem::ClassTypes).IsEmpty() == false? EVisibility::Visible : EVisibility::Collapsed;
+}
+
+FString SNewProjectWizard::GetSelectedTemplateAssetTypes() const
+{
+	return GetSelectedTemplateProperty<FString>(&FTemplateItem::AssetTypes);
+}
+
+EVisibility SNewProjectWizard::GetSelectedTemplateAssetVisibility() const
+{
+	return GetSelectedTemplateProperty<FString>(&FTemplateItem::AssetTypes).IsEmpty() == false ? EVisibility::Visible : EVisibility::Collapsed;
+}
+
 const FSlateBrush* SNewProjectWizard::GetSelectedTemplatePreviewImage() const
 {
 	auto PreviewImage = GetSelectedTemplateProperty(&FTemplateItem::PreviewImage);
@@ -1055,7 +1126,9 @@ void SNewProjectWizard::FindTemplateProjects()
 		TEXT("_1"),			// SortKey
 		TEXT(""),			// No filename, this is a generation template
 		MakeShareable( new FSlateBrush( *FEditorStyle::GetBrush("GameProjectDialog.BlankProjectThumbnail") ) ),
-		MakeShareable( new FSlateBrush( *FEditorStyle::GetBrush("GameProjectDialog.BlankProjectPreview") ) )
+		MakeShareable( new FSlateBrush( *FEditorStyle::GetBrush("GameProjectDialog.BlankProjectPreview") ) ),
+		TEXT(""),		// No class types
+		TEXT("")		// No asset types
 		)) );
 
 	Templates.FindOrAdd(FTemplateCategory::CodeCategoryName).Add(MakeShareable(new FTemplateItem(
@@ -1065,7 +1138,9 @@ void SNewProjectWizard::FindTemplateProjects()
 		TEXT("_2"),			// SortKey
 		TEXT(""),			// No filename, this is a generation template
 		MakeShareable( new FSlateBrush( *FEditorStyle::GetBrush("GameProjectDialog.BasicCodeThumbnail") ) ),
-		MakeShareable( new FSlateBrush( *FEditorStyle::GetBrush("GameProjectDialog.BlankProjectPreview") ) )
+		MakeShareable( new FSlateBrush( *FEditorStyle::GetBrush("GameProjectDialog.BlankProjectPreview") ) ),
+		TEXT(""),		// No class types
+		TEXT("")		// No asset types
 		)) );
 
 	// Now discover and all data driven templates
@@ -1122,6 +1197,8 @@ void SNewProjectWizard::FindTemplateProjects()
 					const FString ProjectFilename = Root / FoundProjectFiles[0];
 					FText TemplateName = TemplateDefs->GetDisplayNameText();
 					FText TemplateDescription = TemplateDefs->GetLocalizedDescription();
+					FString ClassTypes = TemplateDefs->ClassTypes;
+					FString AssetTypes = TemplateDefs->AssetTypes;
 
 					// If no template name was specified for the current culture, just use the project name
 					if ( TemplateName.IsEmpty() )
@@ -1174,7 +1251,9 @@ void SNewProjectWizard::FindTemplateProjects()
 						MoveTemp(SortKey),
 						MoveTemp(ProjectFilename),
 						ThumbnailBrush,
-						PreviewBrush
+						PreviewBrush,
+						ClassTypes,
+						AssetTypes
 					)));
 				}
 			}
diff --git a/Engine/Source/Editor/GameProjectGeneration/Private/SNewProjectWizard.h b/Engine/Source/Editor/GameProjectGeneration/Private/SNewProjectWizard.h
--- a/Engine/Source/Editor/GameProjectGeneration/Private/SNewProjectWizard.h
+++ b/Engine/Source/Editor/GameProjectGeneration/Private/SNewProjectWizard.h
@@ -76,6 +76,11 @@ private:
 	/** Get the visibility for the selected template preview image */
 	EVisibility GetSelectedTemplatePreviewVisibility() const;
 
+	FString GetSelectedTemplateClassTypes() const;
+	EVisibility GetSelectedTemplateClassVisibility() const;
+	FString SNewProjectWizard::GetSelectedTemplateAssetTypes() const;	
+	EVisibility SNewProjectWizard::GetSelectedTemplateAssetVisibility() const;
+
 	/** Gets the assembled project filename with path */
 	FString GetProjectFilenameWithPath() const;
 
