Removed extra qualifier in SNewProjectWizard.h and added comments to the 4 new functions therin

diff --git a/Engine/Source/Editor/GameProjectGeneration/Private/SNewProjectWizard.h b/Engine/Source/Editor/GameProjectGeneration/Private/SNewProjectWizard.h
--- a/Engine/Source/Editor/GameProjectGeneration/Private/SNewProjectWizard.h
+++ b/Engine/Source/Editor/GameProjectGeneration/Private/SNewProjectWizard.h
@@ -76,10 +76,17 @@ private:
 	/** Get the visibility for the selected template preview image */
 	EVisibility GetSelectedTemplatePreviewVisibility() const;
 
+	/** Get a string that details the class types referenced in the selected template */
 	FString GetSelectedTemplateClassTypes() const;
+	
+	/** Get a visiblity of the class types display. If the string is empty this return Collapsed otherwise it will return Visible */
 	EVisibility GetSelectedTemplateClassVisibility() const;
-	FString SNewProjectWizard::GetSelectedTemplateAssetTypes() const;	
-	EVisibility SNewProjectWizard::GetSelectedTemplateAssetVisibility() const;
+	
+	/** Get a string that details the class types referenced in the selected template */
+	FString GetSelectedTemplateAssetTypes() const;	
+	
+	/** Get a visiblity of the asset types display. If the string is empty this return Collapsed otherwise it will return Visible */
+	EVisibility GetSelectedTemplateAssetVisibility() const;
 
 	/** Gets the assembled project filename with path */
 	FString GetProjectFilenameWithPath() const;
