When creating new projects, only set the engine association if they're foreign.

diff --git a/Engine/Source/Editor/GameProjectGeneration/Private/GameProjectUtils.cpp b/Engine/Source/Editor/GameProjectGeneration/Private/GameProjectUtils.cpp
--- a/Engine/Source/Editor/GameProjectGeneration/Private/GameProjectUtils.cpp
+++ b/Engine/Source/Editor/GameProjectGeneration/Private/GameProjectUtils.cpp
@@ -879,7 +879,7 @@ bool GameProjectUtils::GenerateProjectFromScratch(const FProjectInformation& InP
 	// Generate the project file
 	{
 		FText LocalFailReason;
-		if (IProjectManager::Get().GenerateNewProjectFile(InProjectInfo.ProjectFilename, StartupModuleNames, FDesktopPlatformModule::Get()->GetCurrentEngineIdentifier(), LocalFailReason))
+		if (IProjectManager::Get().GenerateNewProjectFile(InProjectInfo.ProjectFilename, StartupModuleNames, TEXT(""), LocalFailReason))
 		{
 			CreatedFiles.Add(InProjectInfo.ProjectFilename);
 		}
@@ -889,6 +889,13 @@ bool GameProjectUtils::GenerateProjectFromScratch(const FProjectInformation& InP
 			DeleteCreatedFiles(NewProjectFolder, CreatedFiles);
 			return false;
 		}
+
+		// Set the engine identifier for it. Do this after saving, so it can be correctly detected as foreign or non-foreign.
+		if(!SetEngineAssociationForForeignProject(InProjectInfo.ProjectFilename, OutFailReason))
+		{
+			DeleteCreatedFiles(NewProjectFolder, CreatedFiles);
+			return false;
+		}
 	}
 
 	SlowTask.EnterProgressFrame();
@@ -1277,7 +1284,7 @@ bool GameProjectUtils::CreateProjectFromTemplate(const FProjectInformation& InPr
 		}
 
 		// Update it to current
-		Project.EngineAssociation = FDesktopPlatformModule::Get()->GetCurrentEngineIdentifier();
+		Project.EngineAssociation.Empty();
 		Project.EpicSampleNameHash = 0;
 
 		// Fix up module names
@@ -1296,6 +1303,13 @@ bool GameProjectUtils::CreateProjectFromTemplate(const FProjectInformation& InPr
 			return false;
 		}
 
+		// Set the engine identifier if it's a foreign project. Do this after saving, so it can be correctly detected as foreign.
+		if(!SetEngineAssociationForForeignProject(InProjectInfo.ProjectFilename, OutFailReason))
+		{
+			DeleteCreatedFiles(DestFolder, CreatedFiles);
+			return false;
+		}
+
 		// Add it to the list of created files
 		CreatedFiles.Add(InProjectInfo.ProjectFilename);
 	}
@@ -1345,6 +1359,19 @@ bool GameProjectUtils::CreateProjectFromTemplate(const FProjectInformation& InPr
 	return true;
 }
 
+bool GameProjectUtils::SetEngineAssociationForForeignProject(const FString& ProjectFileName, FText& OutFailReason)
+{
+	if(FUProjectDictionary(FPaths::RootDir()).IsForeignProject(ProjectFileName))
+	{
+		if(!FDesktopPlatformModule::Get()->SetEngineIdentifierForProject(ProjectFileName, FDesktopPlatformModule::Get()->GetCurrentEngineIdentifier()))
+		{
+			OutFailReason = LOCTEXT("FailedToSetEngineIdentifier", "Couldn't set engine identifier for project");
+			return false;
+		}
+	}
+	return true;
+}
+
 FString GameProjectUtils::GetTemplateDefsFilename()
 {
 	return TEXT("TemplateDefs.ini");
diff --git a/Engine/Source/Editor/GameProjectGeneration/Private/GameProjectUtils.h b/Engine/Source/Editor/GameProjectGeneration/Private/GameProjectUtils.h
--- a/Engine/Source/Editor/GameProjectGeneration/Private/GameProjectUtils.h
+++ b/Engine/Source/Editor/GameProjectGeneration/Private/GameProjectUtils.h
@@ -283,6 +283,9 @@ private:
 	/** Generates a new project using a template project */
 	static bool CreateProjectFromTemplate(const FProjectInformation& InProjectInfo, FText& OutFailReason);
 
+	/** Sets the engine association for a new project. Handles foreign and non-foreign projects. */
+	static bool SetEngineAssociationForForeignProject(const FString& ProjectFileName, FText& OutFailReason);
+
 	/** Copies starter content into the specified project folder. */
 	static bool CopyStarterContent(const FString& DestProjectFolder, FText& OutFailReason);
 
