Switched to using Epic Games Launcher instead of Epic Launcher

diff --git a/Engine/Source/Editor/GameProjectGeneration/Private/SProjectBrowser.cpp b/Engine/Source/Editor/GameProjectGeneration/Private/SProjectBrowser.cpp
--- a/Engine/Source/Editor/GameProjectGeneration/Private/SProjectBrowser.cpp
+++ b/Engine/Source/Editor/GameProjectGeneration/Private/SProjectBrowser.cpp
@@ -1235,7 +1235,7 @@ FReply SProjectBrowser::HandleMarketplaceTabButtonClicked()
 		{
 			EventAttributes.Add(FAnalyticsEventAttribute(TEXT("OpenSucceeded"), TEXT("FALSE")));
 
-			if (EAppReturnType::Yes == FMessageDialog::Open(EAppMsgType::YesNo, LOCTEXT("InstallMarketplacePrompt", "The Marketplace requires the Unreal Engine Launcher, which does not seem to be installed on your computer. Would you like to install it now?")))
+			if (EAppReturnType::Yes == FMessageDialog::Open(EAppMsgType::YesNo, LOCTEXT("InstallMarketplacePrompt", "The Marketplace requires the Epic Games Launcher, which does not seem to be installed on your computer. Would you like to install it now?")))
 			{
 				if (!DesktopPlatform->OpenLauncher(true, TEXT("-OpenMarket")))
 				{
diff --git a/Engine/Source/Editor/LevelEditor/Private/LevelEditorActions.cpp b/Engine/Source/Editor/LevelEditor/Private/LevelEditorActions.cpp
--- a/Engine/Source/Editor/LevelEditor/Private/LevelEditorActions.cpp
+++ b/Engine/Source/Editor/LevelEditor/Private/LevelEditorActions.cpp
@@ -1690,7 +1690,7 @@ void FLevelEditorActionCallbacks::OpenMarketplace()
 		{
 			EventAttributes.Add(FAnalyticsEventAttribute(TEXT("OpenSucceeded"), TEXT("FALSE")));
 
-			if (EAppReturnType::Yes == FMessageDialog::Open(EAppMsgType::YesNo, LOCTEXT("InstallMarketplacePrompt", "The Marketplace requires the Epic Launcher, which does not seem to be installed on your computer. Would you like to install it now?")))
+			if (EAppReturnType::Yes == FMessageDialog::Open(EAppMsgType::YesNo, LOCTEXT("InstallMarketplacePrompt", "The Marketplace requires the Epic Games Launcher, which does not seem to be installed on your computer. Would you like to install it now?")))
 			{
 				if (!DesktopPlatform->OpenLauncher(true, TEXT("-OpenMarket")))
 				{
