Added additional checks in the New Class Dialog to disallow creating classes whose generated headers would clash with existing engine headers.

diff --git a/Engine/Source/Editor/GameProjectGeneration/Private/GameProjectUtils.cpp b/Engine/Source/Editor/GameProjectGeneration/Private/GameProjectUtils.cpp
--- a/Engine/Source/Editor/GameProjectGeneration/Private/GameProjectUtils.cpp
+++ b/Engine/Source/Editor/GameProjectGeneration/Private/GameProjectUtils.cpp
@@ -712,7 +712,7 @@ void GameProjectUtils::OpenAddCodeToProjectDialog()
 	}
 }
 
-bool GameProjectUtils::IsValidClassNameForCreation(const FString& NewClassName, const FModuleContextInfo& ModuleInfo, FText& OutFailReason)
+bool GameProjectUtils::IsValidClassNameForCreation(const FString& NewClassName, const FModuleContextInfo& ModuleInfo, const TSet<FString>& DisallowedHeaderNames, FText& OutFailReason)
 {
 	if ( NewClassName.IsEmpty() )
 	{
@@ -771,12 +771,24 @@ bool GameProjectUtils::IsValidClassNameForCreation(const FString& NewClassName,
 		}
 	}
 
+	// See if header name clashes with an engine header
+	{
+		FString UnusedFoundPath;
+		if (DisallowedHeaderNames.Contains(NewClassName))
+		{
+			FFormatNamedArguments Args;
+			Args.Add(TEXT("NewHeaderName"), FText::FromString(NewClassName + ".h"));
+			OutFailReason = FText::Format(LOCTEXT("HeaderNameAlreadyExists", "The file {NewHeaderName} already exists elsewhere in the engine."), Args);
+			return false;
+		}
+	}
+
 	return true;
 }
 
-bool GameProjectUtils::AddCodeToProject(const FString& NewClassName, const FString& NewClassPath, const FModuleContextInfo& ModuleInfo, const FNewClassInfo ParentClassInfo, FString& OutHeaderFilePath, FString& OutCppFilePath, FText& OutFailReason)
+bool GameProjectUtils::AddCodeToProject(const FString& NewClassName, const FString& NewClassPath, const FModuleContextInfo& ModuleInfo, const FNewClassInfo ParentClassInfo, const TSet<FString>& DisallowedHeaderNames, FString& OutHeaderFilePath, FString& OutCppFilePath, FText& OutFailReason)
 {
-	const bool bAddCodeSuccessful = AddCodeToProject_Internal(NewClassName, NewClassPath, ModuleInfo, ParentClassInfo, OutHeaderFilePath, OutCppFilePath, OutFailReason);
+	const bool bAddCodeSuccessful = AddCodeToProject_Internal(NewClassName, NewClassPath, ModuleInfo, ParentClassInfo, DisallowedHeaderNames, OutHeaderFilePath, OutCppFilePath, OutFailReason);
 
 	if( FEngineAnalytics::IsAvailable() )
 	{
@@ -2799,7 +2811,7 @@ bool GameProjectUtils::ProjectHasCodeFiles()
 	return GameProjectUtils::GetProjectCodeFileCount() > 0;
 }
 
-bool GameProjectUtils::AddCodeToProject_Internal(const FString& NewClassName, const FString& NewClassPath, const FModuleContextInfo& ModuleInfo, const FNewClassInfo ParentClassInfo, FString& OutHeaderFilePath, FString& OutCppFilePath, FText& OutFailReason)
+bool GameProjectUtils::AddCodeToProject_Internal(const FString& NewClassName, const FString& NewClassPath, const FModuleContextInfo& ModuleInfo, const FNewClassInfo ParentClassInfo, const TSet<FString>& DisallowedHeaderNames, FString& OutHeaderFilePath, FString& OutCppFilePath, FText& OutFailReason)
 {
 	if ( !ParentClassInfo.IsSet() )
 	{
@@ -2810,7 +2822,7 @@ bool GameProjectUtils::AddCodeToProject_Internal(const FString& NewClassName, co
 	const FString CleanClassName = ParentClassInfo.GetCleanClassName(NewClassName);
 	const FString FinalClassName = ParentClassInfo.GetFinalClassName(NewClassName);
 
-	if ( !IsValidClassNameForCreation(FinalClassName, ModuleInfo, OutFailReason) )
+	if (!IsValidClassNameForCreation(FinalClassName, ModuleInfo, DisallowedHeaderNames, OutFailReason))
 	{
 		return false;
 	}
@@ -2921,7 +2933,7 @@ bool GameProjectUtils::AddCodeToProject_Internal(const FString& NewClassName, co
 bool GameProjectUtils::FindSourceFileInProject(const FString& InFilename, const FString& InSearchPath, FString& OutPath)
 {
 	TArray<FString> Filenames;
-	const FString FilenameWidcard = TEXT("*") + InFilename;
+	const FString FilenameWidcard = /*TEXT("*") +*/ InFilename;
 	IFileManager::Get().FindFilesRecursive(Filenames, *InSearchPath, *FilenameWidcard, true, false, false);
 	
 	if(Filenames.Num())
diff --git a/Engine/Source/Editor/GameProjectGeneration/Private/GameProjectUtils.h b/Engine/Source/Editor/GameProjectGeneration/Private/GameProjectUtils.h
--- a/Engine/Source/Editor/GameProjectGeneration/Private/GameProjectUtils.h
+++ b/Engine/Source/Editor/GameProjectGeneration/Private/GameProjectUtils.h
@@ -176,10 +176,10 @@ public:
 	static void OpenAddCodeToProjectDialog();
 
 	/** Returns true if the specified class name is properly formed and does not conflict with another class */
-	static bool IsValidClassNameForCreation(const FString& NewClassName, const FModuleContextInfo& ModuleInfo, FText& OutFailReason);
+	static bool IsValidClassNameForCreation(const FString& NewClassName, const FModuleContextInfo& ModuleInfo, const TSet<FString>& DisallowedHeaderNames, FText& OutFailReason);
 
 	/** Adds new source code to the project. When returning true, OutSyncFileAndLineNumber will be the the preferred target file to sync in the users code editing IDE, formatted for use with GenericApplication::GotoLineInSource */
-	static bool AddCodeToProject(const FString& NewClassName, const FString& NewClassPath, const FModuleContextInfo& ModuleInfo, const FNewClassInfo ParentClassInfo, FString& OutHeaderFilePath, FString& OutCppFilePath, FText& OutFailReason);
+	static bool AddCodeToProject(const FString& NewClassName, const FString& NewClassPath, const FModuleContextInfo& ModuleInfo, const FNewClassInfo ParentClassInfo, const TSet<FString>& DisallowedHeaderNames, FString& OutHeaderFilePath, FString& OutCppFilePath, FText& OutFailReason);
 
 	/** Loads a template project definitions object from the TemplateDefs.ini file in the specified project */
 	static UTemplateProjectDefs* LoadTemplateDefs(const FString& ProjectDirectory);
@@ -402,7 +402,7 @@ private:
 	static bool CheckoutGameProjectFile(const FString& ProjectFilename, FText& OutFailReason);
 
 	/** Internal handler for AddCodeToProject*/
-	static bool AddCodeToProject_Internal(const FString& NewClassName, const FString& NewClassPath, const FModuleContextInfo& ModuleInfo, const FNewClassInfo ParentClassInfo, FString& OutHeaderFilePath, FString& OutCppFilePath, FText& OutFailReason);
+	static bool AddCodeToProject_Internal(const FString& NewClassName, const FString& NewClassPath, const FModuleContextInfo& ModuleInfo, const FNewClassInfo ParentClassInfo, const TSet<FString>& DisallowedHeaderNames, FString& OutHeaderFilePath, FString& OutCppFilePath, FText& OutFailReason);
 
 	/** Handler for the user confirming they've read the name legnth warning */
 	static void OnWarningReasonOk();
diff --git a/Engine/Source/Editor/GameProjectGeneration/Private/SNewClassDialog.cpp b/Engine/Source/Editor/GameProjectGeneration/Private/SNewClassDialog.cpp
--- a/Engine/Source/Editor/GameProjectGeneration/Private/SNewClassDialog.cpp
+++ b/Engine/Source/Editor/GameProjectGeneration/Private/SNewClassDialog.cpp
@@ -81,6 +81,17 @@ private:
 	const TSharedPtr<FModuleContextInfo>* SelectedModuleInfoPtr;
 };
 
+static void FindPublicEngineHeaderFiles(TArray<FString>& OutFiles, const FString& Path)
+{
+	TArray<FString> ModuleDirs;
+	IFileManager::Get().FindFiles(ModuleDirs, *(Path / TEXT("*")), false, true);
+	for (const FString& ModuleDir : ModuleDirs)
+	{
+		IFileManager::Get().FindFilesRecursive(OutFiles, *(Path / ModuleDir / TEXT("Classes")), TEXT("*.h"), true, false, false);
+		IFileManager::Get().FindFilesRecursive(OutFiles, *(Path / ModuleDir / TEXT("Public")), TEXT("*.h"), true, false, false);
+	}
+}
+
 BEGIN_SLATE_FUNCTION_BUILD_OPTIMIZATION
 void SNewClassDialog::Construct( const FArguments& InArgs )
 {
@@ -93,6 +104,22 @@ void SNewClassDialog::Construct( const FArguments& InArgs )
 		{
 			AvailableModules.Emplace(MakeShareable(new FModuleContextInfo(ModuleInfo)));
 		}
+
+		// Get collection of disallowed header names
+		TArray<FString> HeaderFiles;
+		FindPublicEngineHeaderFiles(HeaderFiles, FPaths::EngineDir() / TEXT("Source") / TEXT("Developer"));
+		FindPublicEngineHeaderFiles(HeaderFiles, FPaths::EngineDir() / TEXT("Source") / TEXT("Editor"));
+		FindPublicEngineHeaderFiles(HeaderFiles, FPaths::EngineDir() / TEXT("Source") / TEXT("Runtime"));
+
+		for (const FString& HeaderFile : HeaderFiles)
+		{
+			DisallowedHeaderNames.Add(FPaths::GetBaseFilename(HeaderFile));
+		}
+
+		for (TObjectIterator<UClass> ClassIt; ClassIt; ++ClassIt)
+		{
+			DisallowedHeaderNames.Remove(ClassIt->GetName());
+		}
 	}
 
 	// If we have a runtime module with the same name as our project, then use that
@@ -981,7 +1008,7 @@ void SNewClassDialog::FinishClicked()
 	FString CppFilePath;
 
 	FText FailReason;
-	if ( GameProjectUtils::AddCodeToProject(NewClassName, NewClassPath, *SelectedModuleInfo, ParentClassInfo, HeaderFilePath, CppFilePath, FailReason) )
+	if (GameProjectUtils::AddCodeToProject(NewClassName, NewClassPath, *SelectedModuleInfo, ParentClassInfo, DisallowedHeaderNames, HeaderFilePath, CppFilePath, FailReason))
 	{
 		// Prevent periodic validity checks. This is to prevent a brief error message about the class already existing while you are exiting.
 		bPreventPeriodicValidityChecksUntilNextChange = true;
@@ -1211,7 +1238,7 @@ void SNewClassDialog::UpdateInputValidity()
 	// Validate the class name only if the path is valid
 	if ( bLastInputValidityCheckSuccessful )
 	{
-		bLastInputValidityCheckSuccessful = GameProjectUtils::IsValidClassNameForCreation(NewClassName, *SelectedModuleInfo, LastInputValidityErrorText);
+		bLastInputValidityCheckSuccessful = GameProjectUtils::IsValidClassNameForCreation(NewClassName, *SelectedModuleInfo, DisallowedHeaderNames, LastInputValidityErrorText);
 	}
 
 	LastPeriodicValidityCheckTime = FSlateApplication::Get().GetCurrentTime();
diff --git a/Engine/Source/Editor/GameProjectGeneration/Private/SNewClassDialog.h b/Engine/Source/Editor/GameProjectGeneration/Private/SNewClassDialog.h
--- a/Engine/Source/Editor/GameProjectGeneration/Private/SNewClassDialog.h
+++ b/Engine/Source/Editor/GameProjectGeneration/Private/SNewClassDialog.h
@@ -222,4 +222,7 @@ private:
 
 	/** Information about the currently selected module; used for class validation */
 	TSharedPtr<FModuleContextInfo> SelectedModuleInfo;
+
+	/** Collection of disallowed header filenames */
+	TSet<FString> DisallowedHeaderNames;
 };
