Removed SLATE_TEXT_ATTRIBUTE from AppFramework specific widgets

diff --git a/Engine/Source/Editor/GameProjectGeneration/Private/SNewClassDialog.cpp b/Engine/Source/Editor/GameProjectGeneration/Private/SNewClassDialog.cpp
--- a/Engine/Source/Editor/GameProjectGeneration/Private/SNewClassDialog.cpp
+++ b/Engine/Source/Editor/GameProjectGeneration/Private/SNewClassDialog.cpp
@@ -181,8 +181,8 @@ void SNewClassDialog::Construct( const FArguments& InArgs )
 				SAssignNew( MainWizard, SWizard)
 				.ShowPageList(false)
 				.CanFinish(this, &SNewClassDialog::CanFinish)
-				.FinishButtonText( LOCTEXT("FinishButtonText", "Create Class").ToString() )
-				.FinishButtonToolTip ( LOCTEXT("FinishButtonToolTip", "Creates the code files to add your new class.").ToString() )
+				.FinishButtonText( LOCTEXT("FinishButtonText", "Create Class") )
+				.FinishButtonToolTip ( LOCTEXT("FinishButtonToolTip", "Creates the code files to add your new class.") )
 				.OnCanceled(this, &SNewClassDialog::CancelClicked)
 				.OnFinished(this, &SNewClassDialog::FinishClicked)
 				.InitialPageIndex(ParentClassInfo.IsSet() ? 1 : 0)
diff --git a/Engine/Source/Runtime/AppFramework/Public/Widgets/Workflow/SWizard.h b/Engine/Source/Runtime/AppFramework/Public/Widgets/Workflow/SWizard.h
--- a/Engine/Source/Runtime/AppFramework/Public/Widgets/Workflow/SWizard.h
+++ b/Engine/Source/Runtime/AppFramework/Public/Widgets/Workflow/SWizard.h
@@ -153,10 +153,10 @@ public:
 		SLATE_ATTRIBUTE(bool, CanFinish)
 
 		/** Holds the label text of the wizard's 'Finish' button. */
-		SLATE_TEXT_ATTRIBUTE(FinishButtonText)
+		SLATE_ATTRIBUTE(FText, FinishButtonText)
 
 		/** Holds the tool tip text of the wizard's 'Finish' button. */
-		SLATE_TEXT_ATTRIBUTE(FinishButtonToolTip)
+		SLATE_ATTRIBUTE(FText, FinishButtonToolTip)
 
 		/** Holds the index of the initial page to be displayed (0 = default). */
 		SLATE_ATTRIBUTE(int32, InitialPageIndex)
