Fixed code relying on SLATE_TEXT_ATTRIBUTE for SHyperlink

Made sure everything was using FText rather than FString.

diff --git a/Engine/Plugins/Editor/PluginsEditor/Source/PluginsEditor/Private/SPluginListTile.cpp b/Engine/Plugins/Editor/PluginsEditor/Source/PluginsEditor/Private/SPluginListTile.cpp
--- a/Engine/Plugins/Editor/PluginsEditor/Source/PluginsEditor/Private/SPluginListTile.cpp
+++ b/Engine/Plugins/Editor/PluginsEditor/Source/PluginsEditor/Private/SPluginListTile.cpp
@@ -124,7 +124,7 @@ void SPluginListTile::Construct( const FArguments& Args, const TSharedRef<SPlugi
 				.Padding(2.0f, 0.0f, 0.0f, 0.0f)
 				[				
 					SNew(SHyperlink)
-						.Text(Item->PluginStatus.CreatedBy)
+						.Text(FText::FromString(Item->PluginStatus.CreatedBy))
 						.ToolTipText(FText::Format(LOCTEXT("NavigateToCreatedByURL", "Visit the vendor's web site ({0})"), FText::FromString(CreatedByURL)))
 						.OnNavigate_Lambda([=]() { FPlatformProcess::LaunchURL(*CreatedByURL, nullptr, nullptr); })
 				];
diff --git a/Engine/Source/Developer/SlateReflector/Private/Widgets/SWidgetReflectorTreeWidgetItem.cpp b/Engine/Source/Developer/SlateReflector/Private/Widgets/SWidgetReflectorTreeWidgetItem.cpp
--- a/Engine/Source/Developer/SlateReflector/Private/Widgets/SWidgetReflectorTreeWidgetItem.cpp
+++ b/Engine/Source/Developer/SlateReflector/Private/Widgets/SWidgetReflectorTreeWidgetItem.cpp
@@ -41,7 +41,7 @@ TSharedRef<SWidget> SReflectorTreeWidgetItem::GenerateWidgetForColumn(const FNam
 			.Padding(FMargin(2.0f, 0.0f))
 			[
 				SNew(SHyperlink)
-				.Text(this, &SReflectorTreeWidgetItem::GetReadableLocation)
+				.Text(this, &SReflectorTreeWidgetItem::GetReadableLocationAsText)
 				.OnNavigate(this, &SReflectorTreeWidgetItem::HandleHyperlinkNavigate)
 			];
 	}
@@ -136,7 +136,7 @@ void SReflectorTreeWidgetItem::HandleHyperlinkNavigate()
 		{
 			if ( OnAccessSourceCode.IsBound() )
 			{
-				OnAccessSourceCode.Execute(GetWidgetFile(), GetWidgetLineNumber(), 0);
+				OnAccessSourceCode.Execute(GetWidgetFile().ToString(), GetWidgetLineNumber(), 0);
 			}
 		}
 	}
diff --git a/Engine/Source/Developer/SlateReflector/Private/Widgets/SWidgetReflectorTreeWidgetItem.h b/Engine/Source/Developer/SlateReflector/Private/Widgets/SWidgetReflectorTreeWidgetItem.h
--- a/Engine/Source/Developer/SlateReflector/Private/Widgets/SWidgetReflectorTreeWidgetItem.h
+++ b/Engine/Source/Developer/SlateReflector/Private/Widgets/SWidgetReflectorTreeWidgetItem.h
@@ -48,20 +48,25 @@ public:
 protected:
 
 	/** @return String representation of the widget we are visualizing */
-	FString GetWidgetType() const
+	FText GetWidgetType() const
 	{
 		return WidgetInfo.Get()->Widget.IsValid()
-			? WidgetInfo.Get()->Widget.Pin()->GetTypeAsString()
-			: TEXT("Null Widget");
+			? FText::FromString(WidgetInfo.Get()->Widget.Pin()->GetTypeAsString())
+			: NSLOCTEXT("SWidgetReflector", "NullWidget", "Null Widget");
 	}
 	
-	FString GetReadableLocation() const;
+	virtual FString GetReadableLocation() const override;
 
-	FString GetWidgetFile() const
+	FText GetReadableLocationAsText() const
+	{
+		return FText::FromString(GetReadableLocation());
+	}
+
+	FText GetWidgetFile() const
 	{
 		return WidgetInfo.Get()->Widget.IsValid()
-			? WidgetInfo.Get()->Widget.Pin()->GetCreatedInFile()
-			: FString();
+			? FText::FromString(WidgetInfo.Get()->Widget.Pin()->GetCreatedInFile())
+			: FText::GetEmpty();
 	}
 
 	int32 GetWidgetLineNumber() const
@@ -71,12 +76,12 @@ protected:
 			: 0;
 	}
 
-	FString GetVisibilityAsString() const
+	FText GetVisibilityAsString() const
 	{
 		TSharedPtr<SWidget> TheWidget = WidgetInfo.Get()->Widget.Pin();
 		return TheWidget.IsValid()
-			? TheWidget->GetVisibility().ToString()
-			: FString();
+			? FText::FromString(TheWidget->GetVisibility().ToString())
+			: FText::GetEmpty();
 	}
 
 	/** @return The tint of the reflector node */
diff --git a/Engine/Source/Editor/Documentation/Private/UDNParser.cpp b/Engine/Source/Editor/Documentation/Private/UDNParser.cpp
--- a/Engine/Source/Editor/Documentation/Private/UDNParser.cpp
+++ b/Engine/Source/Editor/Documentation/Private/UDNParser.cpp
@@ -923,7 +923,7 @@ TSharedRef< SWidget > FUDNParser::GenerateExcerptContent( const FString& Link, F
 
 				AppendExcerpt(Box,
 					SNew(SHyperlink)
-					.Text(Line.AdditionalContent[0])
+					.Text(FText::FromString(Line.AdditionalContent[0]))
 					.TextStyle(FEditorStyle::Get(), Style.HyperlinkTextStyleName)
 					.UnderlineStyle(FEditorStyle::Get(), Style.HyperlinkButtonStyleName)
 					.OnNavigate( this, &FUDNParser::HandleHyperlinkNavigate, Line.AdditionalContent[1])
diff --git a/Engine/Source/Editor/GameProjectGeneration/Private/SNewClassDialog.cpp b/Engine/Source/Editor/GameProjectGeneration/Private/SNewClassDialog.cpp
--- a/Engine/Source/Editor/GameProjectGeneration/Private/SNewClassDialog.cpp
+++ b/Engine/Source/Editor/GameProjectGeneration/Private/SNewClassDialog.cpp
@@ -762,9 +762,9 @@ TSharedRef<ITableRow> SNewClassDialog::MakeParentClassListViewWidget(TSharedPtr<
 		];
 }
 
-FString SNewClassDialog::GetSelectedParentClassName() const
+FText SNewClassDialog::GetSelectedParentClassName() const
 {
-	return ParentClassInfo.IsSet() ? ParentClassInfo.GetClassName() : TEXT("");
+	return ParentClassInfo.IsSet() ? FText::FromString(ParentClassInfo.GetClassName()) : FText::GetEmpty();
 }
 
 FString GetClassHeaderPath(const UClass* Class)
@@ -785,14 +785,14 @@ EVisibility SNewClassDialog::GetSourceHyperlinkVisibility() const
 	return (GetClassHeaderPath(ParentClassInfo.BaseClass).Len() > 0 ? EVisibility::Visible : EVisibility::Hidden);
 }
 
-FString SNewClassDialog::GetSelectedParentClassFilename() const
+FText SNewClassDialog::GetSelectedParentClassFilename() const
 {
 	const FString ClassHeaderPath = GetClassHeaderPath(ParentClassInfo.BaseClass);
 	if (ClassHeaderPath.Len() > 0)
 	{
-		return FPaths::GetCleanFilename(*ClassHeaderPath);
+		return FText::FromString(FPaths::GetCleanFilename(*ClassHeaderPath));
 	}
-	return FString();
+	return FText::GetEmpty();
 }
 
 EVisibility SNewClassDialog::GetDocLinkVisibility() const
@@ -870,14 +870,14 @@ EVisibility SNewClassDialog::GetNameErrorLabelVisibility() const
 	return GetNameErrorLabelText().IsEmpty() ? EVisibility::Hidden : EVisibility::Visible;
 }
 
-FString SNewClassDialog::GetNameErrorLabelText() const
+FText SNewClassDialog::GetNameErrorLabelText() const
 {
 	if ( !bLastInputValidityCheckSuccessful )
 	{
-		return LastInputValidityErrorText.ToString();
+		return LastInputValidityErrorText;
 	}
 
-	return TEXT("");
+	return FText::GetEmpty();
 }
 
 EVisibility SNewClassDialog::GetGlobalErrorLabelVisibility() const
@@ -919,15 +919,17 @@ void SNewClassDialog::OnNamePageEntered()
 	FSlateApplication::Get().SetKeyboardFocus(ClassNameEditBox, EFocusCause::SetDirectly);
 }
 
-FString SNewClassDialog::GetNameClassTitle() const
+FText SNewClassDialog::GetNameClassTitle() const
 {
-	FString ParentClassName = GetSelectedParentClassName();
-	if(ParentClassName.IsEmpty() || ParentClassName == "None")
+	static const FString NoneString = TEXT("None");
+
+	const FText ParentClassName = GetSelectedParentClassName();
+	if(!ParentClassName.IsEmpty() && ParentClassName.ToString() != NoneString)
 	{
-		ParentClassName = TEXT("Class");
+		return FText::Format( LOCTEXT( "NameClassTitle", "Name Your New {0}" ), ParentClassName );
 	}
 
-	return FText::Format( LOCTEXT( "NameClassTitle", "Name Your New {0}" ), FText::FromString(ParentClassName) ).ToString();
+	return LOCTEXT( "NameClassGenericTitle", "Name Your New Class" );
 }
 
 FText SNewClassDialog::OnGetClassNameText() const
diff --git a/Engine/Source/Editor/GameProjectGeneration/Private/SNewClassDialog.h b/Engine/Source/Editor/GameProjectGeneration/Private/SNewClassDialog.h
--- a/Engine/Source/Editor/GameProjectGeneration/Private/SNewClassDialog.h
+++ b/Engine/Source/Editor/GameProjectGeneration/Private/SNewClassDialog.h
@@ -30,10 +30,10 @@ private:
 	TSharedRef<ITableRow> MakeParentClassListViewWidget(TSharedPtr<FParentClassItem> ParentClassItem, const TSharedRef<STableViewBase>& OwnerTable);
 
 	/** Gets the currently selected parent class name */
-	FString GetSelectedParentClassName() const;
+	FText GetSelectedParentClassName() const;
 
 	/** Gets the currently selected parent class's filename */
-	FString GetSelectedParentClassFilename() const;
+	FText GetSelectedParentClassFilename() const;
 
 	/** Whether the hyper link to go to source should be visible */
 	EVisibility GetSourceHyperlinkVisibility() const;
@@ -72,7 +72,7 @@ private:
 	EVisibility GetNameErrorLabelVisibility() const;
 
 	/** Gets the text to display in the name error label */
-	FString GetNameErrorLabelText() const;
+	FText GetNameErrorLabelText() const;
 
 	/** Gets the visibility of the global error label */
 	EVisibility GetGlobalErrorLabelVisibility() const;
@@ -87,7 +87,7 @@ private:
 	void OnNamePageEntered();
 
 	/** Returns the title text for the "name class" page */
-	FString GetNameClassTitle() const;
+	FText GetNameClassTitle() const;
 
 	/** Returns the text in the class name edit box */
 	FText OnGetClassNameText() const;
diff --git a/Engine/Source/Editor/Kismet/Private/Debugging/SKismetDebuggingView.cpp b/Engine/Source/Editor/Kismet/Private/Debugging/SKismetDebuggingView.cpp
--- a/Engine/Source/Editor/Kismet/Private/Debugging/SKismetDebuggingView.cpp
+++ b/Engine/Source/Editor/Kismet/Private/Debugging/SKismetDebuggingView.cpp
@@ -497,7 +497,7 @@ protected:
 	}
 	
 	
-	FString GetLocationDescription() const;
+	FText GetLocationDescription() const;
 	FReply OnUserToggledEnabled();
 
 	void OnNavigateToBreakpointLocation();
@@ -506,13 +506,13 @@ protected:
 	FString GetStatusTooltip() const;
 };
 
-FString FBreakpointLineItem::GetLocationDescription() const
+FText FBreakpointLineItem::GetLocationDescription() const
 {
 	if (UBreakpoint* MyBreakpoint = BreakpointRef.Get())
 	{
 		return MyBreakpoint->GetLocationDescription();
 	}
-	return TEXT("");
+	return FText::GetEmpty();
 }
 
 FReply FBreakpointLineItem::OnUserToggledEnabled()
@@ -847,13 +847,13 @@ protected:
 		return TEXT("");
 	}
 
-	FString GetContextObjectName() const
+	FText GetContextObjectName() const
 	{
 		const TSimpleRingBuffer<FKismetTraceSample>& TraceStack = FKismetDebugUtilities::GetTraceStack();
 
 		UObject* ObjectContext = (StackIndex < TraceStack.Num()) ? TraceStack(StackIndex).Context.Get() : NULL;
 		
-		return (ObjectContext != NULL) ? ObjectContext->GetName() : LOCTEXT("ObjectDoesNotExist", "(object no longer exists)").ToString();
+		return (ObjectContext != NULL) ? FText::FromString(ObjectContext->GetName()) : LOCTEXT("ObjectDoesNotExist", "(object no longer exists)");
 	}
 
 	void OnNavigateToNode()
diff --git a/Engine/Source/Runtime/Engine/Classes/Engine/Breakpoint.h b/Engine/Source/Runtime/Engine/Classes/Engine/Breakpoint.h
--- a/Engine/Source/Runtime/Engine/Classes/Engine/Breakpoint.h
+++ b/Engine/Source/Runtime/Engine/Classes/Engine/Breakpoint.h
@@ -50,8 +50,7 @@ public:
 	}
 
 	/** Gets a string that describes the location */
-	FString GetLocationDescription() const;
-
+	FText GetLocationDescription() const;
 
 	friend class FKismetDebugUtilities;
 };
diff --git a/Engine/Source/Runtime/Engine/Private/Breakpoint.cpp b/Engine/Source/Runtime/Engine/Private/Breakpoint.cpp
--- a/Engine/Source/Runtime/Engine/Private/Breakpoint.cpp
+++ b/Engine/Source/Runtime/Engine/Private/Breakpoint.cpp
@@ -14,7 +14,7 @@ UBreakpoint::UBreakpoint(const FObjectInitializer& ObjectInitializer)
 	bStepOnce_RemoveAfterHit = false;
 }
 
-FString UBreakpoint::GetLocationDescription() const
+FText UBreakpoint::GetLocationDescription() const
 {
 #if WITH_EDITORONLY_DATA
 	if (Node != NULL)
@@ -30,13 +30,13 @@ FString UBreakpoint::GetLocationDescription() const
 			Result += Node->NodeComment;
 		}
 
-		return Result;
+		return FText::FromString(Result);
 	}
 	else
 	{
-		return TEXT("Error: Invalid location");
+		return NSLOCTEXT("UBreakpoint", "ErrorInvalidLocation", "Error: Invalid location");
 	}
 #else	//#if WITH_EDITORONLY_DATA
-	return TEXT("--- NO EDITOR DATA! ---");
+	return NSLOCTEXT("UBreakpoint", "NoEditorData", "--- NO EDITOR DATA! ---");
 #endif	//#if WITH_EDITORONLY_DATA
 }
