Revised starter content to use asset pack
Moved starter content assets
Added mobile starter content assets
Starter content uses mobile starter content based on target platform

diff --git a/Engine/Source/Editor/AddContentDialog/Private/ContentSourceProviders/FeaturePack/FeaturePackContentSource.cpp b/Engine/Source/Editor/AddContentDialog/Private/ContentSourceProviders/FeaturePack/FeaturePackContentSource.cpp
--- a/Engine/Source/Editor/AddContentDialog/Private/ContentSourceProviders/FeaturePack/FeaturePackContentSource.cpp
+++ b/Engine/Source/Editor/AddContentDialog/Private/ContentSourceProviders/FeaturePack/FeaturePackContentSource.cpp
@@ -6,20 +6,28 @@
 #include "AssetToolsModule.h"
 #include "IPlatformFilePak.h"
 
+DEFINE_LOG_CATEGORY_STATIC(LogFeaturePack, Log, All);
+
+
 FFeaturePackContentSource::FFeaturePackContentSource(FString InFeaturePackPath)
 {
 	FeaturePackPath = InFeaturePackPath;
-
+	bPackValid = false;
 	// Create a pak platform file and mount the feature pack file.
 	FPakPlatformFile PakPlatformFile;
 	FString CommandLine;
 	PakPlatformFile.Initialize(&FPlatformFileManager::Get().GetPlatformFile(), TEXT(""));
 	FString MountPoint = "root:/";
 	PakPlatformFile.Mount(*InFeaturePackPath, 0, *MountPoint);
-
+	
 	// Gets the manifest file as a JSon string
 	TArray<uint8> ManifestBuffer;
-	LoadPakFileToBuffer(PakPlatformFile, FPaths::Combine(*MountPoint, TEXT("manifest.json")), ManifestBuffer);
+	if( LoadPakFileToBuffer(PakPlatformFile, FPaths::Combine(*MountPoint, TEXT("manifest.json")), ManifestBuffer) == false )
+	{
+		UE_LOG(LogFeaturePack, Warning, TEXT("Error in Feature pack %s. Cannot find manifest."), *InFeaturePackPath);
+		Category = EContentSourceCategory::Unknown;
+		return;
+	}
 	FString ManifestString;
 	FFileHelper::BufferToString(ManifestString, ManifestBuffer.GetData(), ManifestBuffer.Num());
 
@@ -52,6 +60,10 @@ FFeaturePackContentSource::FFeaturePackContentSource(FString InFeaturePackPath)
 	{
 		Category = EContentSourceCategory::BlueprintFeature;
 	}
+	else if (CategoryString == "Content")
+	{
+		Category = EContentSourceCategory::Content;
+	}
 	else
 	{
 		Category = EContentSourceCategory::Unknown;
@@ -70,13 +82,19 @@ FFeaturePackContentSource::FFeaturePackContentSource(FString InFeaturePackPath)
 		LoadPakFileToBuffer(PakPlatformFile, FPaths::Combine(*MountPoint, TEXT("Media"), *ScreenshotFilename->AsString()), *SingleScreenshotData);
 		ScreenshotData.Add(MakeShareable(new FImageData(ScreenshotFilename->AsString(), SingleScreenshotData)));
 	}
+	bPackValid = true;
 }
 
-void FFeaturePackContentSource::LoadPakFileToBuffer(FPakPlatformFile& PakPlatformFile, FString Path, TArray<uint8>& Buffer)
+bool FFeaturePackContentSource::LoadPakFileToBuffer(FPakPlatformFile& PakPlatformFile, FString Path, TArray<uint8>& Buffer)
 {
+	bool bResult = false;
 	TSharedPtr<IFileHandle> FileHandle(PakPlatformFile.OpenRead(*Path));
-	Buffer.AddUninitialized(FileHandle->Size());
-	FileHandle->Read(Buffer.GetData(), FileHandle->Size());
+	if( FileHandle.IsValid())
+	{
+		Buffer.AddUninitialized(FileHandle->Size());
+		bResult = FileHandle->Read(Buffer.GetData(), FileHandle->Size());
+	}
+	return bResult;
 }
 
 TArray<FLocalizedText> FFeaturePackContentSource::GetLocalizedNames()
@@ -104,14 +122,41 @@ TArray<TSharedPtr<FImageData>> FFeaturePackContentSource::GetScreenshotData()
 	return ScreenshotData;
 }
 
-void FFeaturePackContentSource::InstallToProject(FString InstallPath)
+bool FFeaturePackContentSource::InstallToProject(FString InstallPath)
 {
-	FAssetToolsModule& AssetToolsModule = FModuleManager::Get().LoadModuleChecked<FAssetToolsModule>("AssetTools");
-	TArray<FString> AssetPaths;
-	AssetPaths.Add(FeaturePackPath);
-	AssetToolsModule.Get().ImportAssets(AssetPaths, InstallPath);
+	bool bResult = false;
+	if( IsDataValid() == false )
+	{
+		UE_LOG(LogFeaturePack, Warning, TEXT("Trying to install invalid pack %s"), *InstallPath);
+	}
+	else
+	{
+		FAssetToolsModule& AssetToolsModule = FModuleManager::Get().LoadModuleChecked<FAssetToolsModule>("AssetTools");
+		TArray<FString> AssetPaths;
+		AssetPaths.Add(FeaturePackPath);
+		TArray<UObject*> ImportedObjects = AssetToolsModule.Get().ImportAssets(AssetPaths, InstallPath);
+		if( ImportedObjects.Num() == 0 )
+		{
+			UE_LOG(LogFeaturePack, Warning, TEXT("No objects imported installing pack %s"), *InstallPath);
+		}
+		else
+		{
+			bResult = true;
+		}
+	}
+	return bResult;
 }
 
 FFeaturePackContentSource::~FFeaturePackContentSource()
 {
+}
+
+bool FFeaturePackContentSource::IsDataValid() const
+{
+	if( bPackValid == false )
+	{
+		return false;
+	}
+	// To Do maybe validate other data here
+	return true;	
 }
\ No newline at end of file
diff --git a/Engine/Source/Editor/AddContentDialog/Private/ContentSourceProviders/FeaturePack/FeaturePackContentSource.h b/Engine/Source/Editor/AddContentDialog/Private/ContentSourceProviders/FeaturePack/FeaturePackContentSource.h
--- a/Engine/Source/Editor/AddContentDialog/Private/ContentSourceProviders/FeaturePack/FeaturePackContentSource.h
+++ b/Engine/Source/Editor/AddContentDialog/Private/ContentSourceProviders/FeaturePack/FeaturePackContentSource.h
@@ -18,12 +18,13 @@ public:
 	virtual TSharedPtr<FImageData> GetIconData() override;
 	virtual TArray<TSharedPtr<FImageData>> GetScreenshotData() override;
 
-	virtual void InstallToProject(FString InstallPath) override;
+	virtual bool InstallToProject(FString InstallPath) override;
+	virtual bool IsDataValid() const override;
 
 	virtual ~FFeaturePackContentSource();
-
+	
 private:
-	void LoadPakFileToBuffer(FPakPlatformFile& PakPlatformFile, FString Path, TArray<uint8>& Buffer);
+	bool LoadPakFileToBuffer(FPakPlatformFile& PakPlatformFile, FString Path, TArray<uint8>& Buffer);
 
 private:
 	FString FeaturePackPath;
@@ -32,4 +33,5 @@ private:
 	EContentSourceCategory Category;
 	TSharedPtr<FImageData> IconData;
 	TArray<TSharedPtr<FImageData>> ScreenshotData;
+	bool bPackValid;
 };
\ No newline at end of file
diff --git a/Engine/Source/Editor/AddContentDialog/Private/ContentSourceProviders/FeaturePack/FeaturePackContentSourceProvider.cpp b/Engine/Source/Editor/AddContentDialog/Private/ContentSourceProviders/FeaturePack/FeaturePackContentSourceProvider.cpp
--- a/Engine/Source/Editor/AddContentDialog/Private/ContentSourceProviders/FeaturePack/FeaturePackContentSourceProvider.cpp
+++ b/Engine/Source/Editor/AddContentDialog/Private/ContentSourceProviders/FeaturePack/FeaturePackContentSourceProvider.cpp
@@ -34,7 +34,11 @@ FFeaturePackContentSourceProvider::FFeaturePackContentSourceProvider()
 	PlatformFile.IterateDirectory(*FeaturePackPath, DirectoryVisitor);
 	for (auto FeaturePackFile : DirectoryVisitor.Files)
 	{
-		ContentSources.Add(MakeShareable(new FFeaturePackContentSource(FeaturePackFile)));
+		TUniquePtr<FFeaturePackContentSource> NewContentSource = MakeUnique<FFeaturePackContentSource>(FeaturePackFile);
+		if(NewContentSource->IsDataValid())
+		{
+			ContentSources.Add(MakeShareable(NewContentSource.Release()));
+		}
 	}
 }
 
diff --git a/Engine/Source/Editor/AddContentDialog/Private/IContentSource.h b/Engine/Source/Editor/AddContentDialog/Private/IContentSource.h
--- a/Engine/Source/Editor/AddContentDialog/Private/IContentSource.h
+++ b/Engine/Source/Editor/AddContentDialog/Private/IContentSource.h
@@ -7,6 +7,7 @@ enum class EContentSourceCategory
 {
 	BlueprintFeature,
 	CodeFeature,
+	Content,
 	Unknown
 };
 
@@ -23,7 +24,7 @@ public:
 		Data = InData;
 	}
 
-	/** Gets the name which identifies the image represented by this obejct. */
+	/** Gets the name which identifies the image represented by this object. */
 	const FString GetName() const
 	{
 		return Name;
@@ -93,8 +94,16 @@ public:
 	/** Gets an array or image data for screenshots for the content source. */
 	virtual TArray<TSharedPtr<FImageData>> GetScreenshotData() = 0;
 
-	/** Installs the content in the content source to the specific path. */
-	virtual void InstallToProject(FString InstallPath) = 0;
+	/*
+	 * Installs the content in the content source to the specific path. 
+	 * @returns true if install suceeded
+	 */
+	virtual bool InstallToProject(FString InstallPath) = 0;
+
+	/** Is the data in this content valid. */
+	virtual bool IsDataValid() const = 0;
 
 	virtual ~IContentSource() { };
+
+
 };
diff --git a/Engine/Source/Editor/GameProjectGeneration/Classes/TemplateProjectDefs.h b/Engine/Source/Editor/GameProjectGeneration/Classes/TemplateProjectDefs.h
--- a/Engine/Source/Editor/GameProjectGeneration/Classes/TemplateProjectDefs.h
+++ b/Engine/Source/Editor/GameProjectGeneration/Classes/TemplateProjectDefs.h
@@ -96,6 +96,10 @@ class UTemplateProjectDefs : public UObject
 	UPROPERTY(config)
 	FString AssetTypes;
 
+	/* Optional list of feature packs to include */
+	UPROPERTY(config)
+	TArray<FString> PacksToInclude;
+
 	/** Fixes up all strings in this definitions object to replace \%TEMPLATENAME\% with the supplied template name and \%PROJECTNAME\% with the supplied project name */
 	void FixupStrings(const FString& TemplateName, const FString& ProjectName);
 
diff --git a/Engine/Source/Editor/GameProjectGeneration/GameProjectGeneration.Build.cs b/Engine/Source/Editor/GameProjectGeneration/GameProjectGeneration.Build.cs
--- a/Engine/Source/Editor/GameProjectGeneration/GameProjectGeneration.Build.cs
+++ b/Engine/Source/Editor/GameProjectGeneration/GameProjectGeneration.Build.cs
@@ -32,7 +32,8 @@ public class GameProjectGeneration : ModuleRules
  				"TargetPlatform",
 				"UnrealEd",
 				"DesktopPlatform",
-                "HardwareTargeting"
+                "HardwareTargeting",                
+                "AssetTools",
 			}
 		);
 
diff --git a/Engine/Source/Editor/GameProjectGeneration/Private/GameProjectUtils.cpp b/Engine/Source/Editor/GameProjectGeneration/Private/GameProjectUtils.cpp
--- a/Engine/Source/Editor/GameProjectGeneration/Private/GameProjectUtils.cpp
+++ b/Engine/Source/Editor/GameProjectGeneration/Private/GameProjectUtils.cpp
@@ -20,7 +20,7 @@
 #include "DesktopPlatformModule.h"
 #include "SNotificationList.h"
 #include "NotificationManager.h"
-#include "GameFramework/GameMode.h"
+#include "AssetToolsModule.h"
 
 #define LOCTEXT_NAMESPACE "GameProjectUtils"
 
@@ -505,58 +505,11 @@ bool GameProjectUtils::OpenCodeIDE(const FString& ProjectFile, FText& OutFailRea
 
 void GameProjectUtils::GetStarterContentFiles(TArray<FString>& OutFilenames)
 {
-	FString const SrcFolder = FPaths::StarterContentDir();
-	FString const ContentFolder = SrcFolder / TEXT("Content");
-
-	// only copying /Content
-	IFileManager::Get().FindFilesRecursive(OutFilenames, *ContentFolder, TEXT("*"), /*Files=*/true, /*Directories=*/false);
-}
-
-bool GameProjectUtils::CopyStarterContent(const FString& DestProjectFolder, FText& OutFailReason)
-{
-	FString const SrcFolder = FPaths::StarterContentDir();
-
-	TArray<FString> FilesToCopy;
-	GetStarterContentFiles(FilesToCopy);
-
-	FScopedSlowTask SlowTask(FilesToCopy.Num(), LOCTEXT("CreatingProjectStatus_CopyingFiles", "Copying Files {SrcFilename}..."));
-	SlowTask.MakeDialog();
-
-	TArray<FString> CreatedFiles;
-	for (FString SrcFilename : FilesToCopy)
-	{
-		// Update the slow task dialog
-		FFormatNamedArguments Args;
-		Args.Add(TEXT("SrcFilename"), FText::FromString(FPaths::GetCleanFilename(SrcFilename)));
-		SlowTask.EnterProgressFrame(1, FText::Format(LOCTEXT("CreatingProjectStatus_CopyingFile", "Copying File {SrcFilename}..."), Args));
-
-		FString FileRelPath = FPaths::GetPath(SrcFilename);
-		FPaths::MakePathRelativeTo(FileRelPath, *SrcFolder);
-
-		// Perform the copy. For file collisions, leave existing file.
-		const FString DestFilename = DestProjectFolder + TEXT("/") + FileRelPath + TEXT("/") + FPaths::GetCleanFilename(SrcFilename);
-		if (!FPaths::FileExists(DestFilename))
-		{
-			if (IFileManager::Get().Copy(*DestFilename, *SrcFilename, false) == COPY_OK)
-			{
-				CreatedFiles.Add(DestFilename);
-			}
-			else
-			{
-				FFormatNamedArguments FailArgs;
-				FailArgs.Add(TEXT("SrcFilename"), FText::FromString(SrcFilename));
-				FailArgs.Add(TEXT("DestFilename"), FText::FromString(DestFilename));
-				OutFailReason = FText::Format(LOCTEXT("FailedToCopyFile", "Failed to copy \"{SrcFilename}\" to \"{DestFilename}\"."), FailArgs);
-				DeleteCreatedFiles(DestProjectFolder, CreatedFiles);
-				return false;
-			}
-		}
-	}
-
-	return true;
+	FString const SrcFolder = FPaths::FeaturePackDir();
+	
+	IFileManager::Get().FindFilesRecursive(OutFilenames, *SrcFolder, TEXT("*.upack"), /*Files=*/true, /*Directories=*/false);
 }
 
-
 bool GameProjectUtils::CreateProject(const FProjectInformation& InProjectInfo, FText& OutFailReason)
 {
 	if ( !IsValidProjectFileForCreation(InProjectInfo.ProjectFilename, OutFailReason) )
@@ -852,6 +805,9 @@ bool GameProjectUtils::GenerateProjectFromScratch(const FProjectInformation& InP
 		return false;
 	}
 
+	// Insert any required feature packs (EG starter content) into ini file. These will be imported automatically when the editor is first run
+	InsertFeaturePacksIntoINIFile(InProjectInfo, OutFailReason);
+	
 	// Make the Content folder
 	const FString ContentFolder = NewProjectFolder / TEXT("Content");
 	if ( !IFileManager::Get().MakeDirectory(*ContentFolder) )
@@ -926,17 +882,6 @@ bool GameProjectUtils::GenerateProjectFromScratch(const FProjectInformation& InP
 
 	SlowTask.EnterProgressFrame();
 
-	if (InProjectInfo.bCopyStarterContent)
-	{
-		// Copy the starter content
-		if ( !CopyStarterContent(NewProjectFolder, OutFailReason) )
-		{
-			DeleteGeneratedProjectFiles(InProjectInfo.ProjectFilename);
-			DeleteCreatedFiles(NewProjectFolder, CreatedFiles);
-			return false;
-		}
-	}
-
 	UE_LOG(LogGameProjectGeneration, Log, TEXT("Created new project with %d files (plus project files)"), CreatedFiles.Num());
 	return true;
 }
@@ -1284,6 +1229,9 @@ bool GameProjectUtils::CreateProjectFromTemplate(const FProjectInformation& InPr
 		}
 	}
 
+	// Insert any required feature packs (EG starter content) into ini file. These will be imported automatically when the editor is first run
+	InsertFeaturePacksIntoINIFile(InProjectInfo, OutFailReason);
+	
 	SlowTask.EnterProgressFrame();
 
 	// Generate the project file
@@ -1351,17 +1299,6 @@ bool GameProjectUtils::CreateProjectFromTemplate(const FProjectInformation& InPr
 
 	SlowTask.EnterProgressFrame();
 
-	if (InProjectInfo.bCopyStarterContent)
-	{
-		// Copy the starter content
-		if ( !CopyStarterContent(DestFolder, OutFailReason) )
-		{
-			DeleteGeneratedProjectFiles(InProjectInfo.ProjectFilename);
-			DeleteCreatedFiles(DestFolder, CreatedFiles);
-			return false;
-		}
-	}
-
 	if (!TemplateDefs->PostGenerateProject(DestFolder, SrcFolder, InProjectInfo.ProjectFilename, InProjectInfo.TemplateFile, InProjectInfo.bShouldGenerateCode, OutFailReason))
 	{
 		DeleteGeneratedProjectFiles(InProjectInfo.ProjectFilename);
@@ -1543,66 +1480,23 @@ bool GameProjectUtils::GenerateConfigFiles(const FProjectInformation& InProjectI
 		
 		if (InProjectInfo.bCopyStarterContent)
 		{
-			FString StarterContentContentDir = FPaths::StarterContentDir() + TEXT("Content/");
-
-			TArray<FString> StarterContentMapFiles;
-			const FString FileWildcard = FString(TEXT("*")) + FPackageName::GetMapPackageExtension();
-		
 			FString SpecificEditorStartupMap;
-			FString SpecificGameDefaultMap;	
-			FString FullEditorStartupMapPath;
-			FString FullGameDefaultMapPath;
+			FString SpecificGameDefaultMap;
 
-			// First we check if there are maps specified in the DefaultEngine.ini in our starter content folder			
-			const FString StarterContentDefaultEngineIniFilename = FPaths::StarterContentDir() / TEXT("Config/DefaultEngine.ini");
-			if (FPaths::FileExists(StarterContentDefaultEngineIniFilename))
+			// If we have starter content packs available, specify starter map
+			if( IsStarterContentAvailableForNewProjects() == true )
 			{
-				FString StarterFileContents;
-				if (FFileHelper::LoadFileToString(StarterFileContents, *StarterContentDefaultEngineIniFilename))
+				if (InProjectInfo.TargetedHardware == EHardwareClass::Mobile)
 				{
-					TArray<FString> StarterIniLines;
-					StarterFileContents.ParseIntoArrayLines(&StarterIniLines);
-					for (int32 Line = 0; Line < StarterIniLines.Num();Line++)
-					{
-						FString EachLine = StarterIniLines[Line];
-						if (EachLine.StartsWith(TEXT("EditorStartupMap")))
-						{
-							EachLine.Split("=", nullptr, &SpecificEditorStartupMap);
-							FullEditorStartupMapPath = (StarterContentContentDir / SpecificEditorStartupMap) + FPackageName::GetMapPackageExtension();
-							FullEditorStartupMapPath = FullEditorStartupMapPath.Replace(TEXT("Game/"), TEXT(""));
-						}
-						if (EachLine.StartsWith(TEXT("GameDefaultMap")))
-						{
-							EachLine.Split("=", nullptr, &SpecificGameDefaultMap);
-							FullGameDefaultMapPath = (StarterContentContentDir / SpecificEditorStartupMap) + FPackageName::GetMapPackageExtension();
-							FullGameDefaultMapPath = FullGameDefaultMapPath.Replace(TEXT("Game/"), TEXT(""));
-						}
-					}					
+					SpecificEditorStartupMap = TEXT("/Game/MobileStarterContent/Maps/StarterMap");
+					SpecificGameDefaultMap = TEXT("/Game/MobileStarterContent/Maps/StarterMap");
 				}
-			}
-
-			// Look for maps in the content folder. If we don't specify maps for EditorStartup and GameDefault we will use the first we find in here
-			IFileManager::Get().FindFilesRecursive(StarterContentMapFiles, *FPaths::StarterContentDir(), *FileWildcard, /*Files=*/true, /*Directories=*/false);			
-			FString MapPackagePath;
-			if (StarterContentMapFiles.Num() > 0)
-			{
-				const FString BaseMapFilename = FPaths::GetBaseFilename(StarterContentMapFiles[0]);
-
-				FString MapPathRelToContent = FPaths::GetPath(StarterContentMapFiles[0]);
-				FPaths::MakePathRelativeTo(MapPathRelToContent, *StarterContentContentDir);
-
-				MapPackagePath = FString(TEXT("/Game/")) + MapPathRelToContent + TEXT("/") + BaseMapFilename;
-			}
-
-			// if either the files we specified don't exist or we didn't specify any, use the first map file we found in the content folder.
-			if (SpecificEditorStartupMap.IsEmpty() || FPaths::FileExists(FullEditorStartupMapPath) == false)
-			{
-				SpecificEditorStartupMap = MapPackagePath;
-			}
-			if (SpecificGameDefaultMap.IsEmpty() || FPaths::FileExists(FullGameDefaultMapPath) == false)
-			{
-				SpecificGameDefaultMap = MapPackagePath;
-			}
+				else
+				{
+					SpecificEditorStartupMap = TEXT("/Game/StarterContent/Maps/StarterMap");
+					SpecificGameDefaultMap = TEXT("/Game/StarterContent/Maps/StarterMap");
+				}
+			}						
 			
 			// Write out the settings for startup map and game default map
 			FileContents += TEXT("[/Script/EngineSettings.GameMapsSettings]") LINE_TERMINATOR;
@@ -1823,8 +1717,9 @@ bool GameProjectUtils::IsStarterContentAvailableForNewProjects()
 {
 	TArray<FString> StarterContentFiles;
 	GetStarterContentFiles(StarterContentFiles);
-
-	return (StarterContentFiles.Num() > 0);
+	
+	bool bHasStaterContent = StarterContentFiles.FindByPredicate([&](const FString& Str){ return Str.Contains("StarterContent"); }) != nullptr;
+	return bHasStaterContent;
 }
 
 TArray<FModuleContextInfo> GameProjectUtils::GetCurrentProjectModules()
@@ -2946,4 +2841,63 @@ bool GameProjectUtils::FindSourceFileInProject(const FString& InFilename, const
 	return false;
 }
 
+bool GameProjectUtils::InsertFeaturePacksIntoINIFile(const FProjectInformation& InProjectInfo, FText& OutFailReason)
+{
+	bool bSuccessfullyProcessed = false;
+	const FString ProjectName = FPaths::GetBaseFilename(InProjectInfo.ProjectFilename);
+	const FString TemplateName = FPaths::GetBaseFilename(InProjectInfo.TemplateFile);
+	const FString SrcFolder = FPaths::GetPath(InProjectInfo.TemplateFile);
+	const FString DestFolder = FPaths::GetPath(InProjectInfo.ProjectFilename);
+	
+	const FString ProjectConfigPath = DestFolder / TEXT("Config");
+	const FString IniFilename = ProjectConfigPath / TEXT("DefaultGame.ini");
+		
+	TArray<FString> PackList;
+	
+	// First the starter content
+	if (InProjectInfo.bCopyStarterContent)
+	{
+		FString StarterPack;
+		if( InProjectInfo.TargetedHardware == EHardwareClass::Mobile )
+		{
+			StarterPack = TEXT("InsertPack=(PackSource=\"MobileStarterContent.upack\",PackName=\"StarterContent\")");
+		}
+		else
+		{
+			StarterPack = TEXT("InsertPack=(PackSource=\"StarterContent.upack\",PackName=\"StarterContent\")");
+		}
+		PackList.Add(StarterPack);
+	}
+
+	// Now any packs specfied in the template def.
+	UTemplateProjectDefs* TemplateDefs = LoadTemplateDefs(SrcFolder);
+	if (TemplateDefs != NULL)
+	{
+		for (int32 iPack = 0; iPack < TemplateDefs->PacksToInclude.Num(); ++iPack)
+		{
+			PackList.Add(TemplateDefs->PacksToInclude[iPack]);
+		}
+	}
+
+	FString FileOutput;
+	if( PackList.Num() != 0)
+	{
+		FileOutput = TEXT("[StartupActions]");
+		FileOutput += LINE_TERMINATOR;
+		FileOutput += TEXT("bAddPacks=True");
+		FileOutput += LINE_TERMINATOR;
+		for (int32 iLine = 0; iLine < PackList.Num(); ++iLine)
+		{
+			FileOutput += PackList[iLine] + LINE_TERMINATOR;
+			
+		}
+		if (FFileHelper::SaveStringToFile(FileOutput, *IniFilename))
+		{
+			bSuccessfullyProcessed = true;
+		}
+	}
+	return bSuccessfullyProcessed;
+}
+
+
 #undef LOCTEXT_NAMESPACE
diff --git a/Engine/Source/Editor/GameProjectGeneration/Private/GameProjectUtils.h b/Engine/Source/Editor/GameProjectGeneration/Private/GameProjectUtils.h
--- a/Engine/Source/Editor/GameProjectGeneration/Private/GameProjectUtils.h
+++ b/Engine/Source/Editor/GameProjectGeneration/Private/GameProjectUtils.h
@@ -286,9 +286,6 @@ private:
 	/** Sets the engine association for a new project. Handles foreign and non-foreign projects. */
 	static bool SetEngineAssociationForForeignProject(const FString& ProjectFileName, FText& OutFailReason);
 
-	/** Copies starter content into the specified project folder. */
-	static bool CopyStarterContent(const FString& DestProjectFolder, FText& OutFailReason);
-
 	/** Returns list of starter content files */
 	static void GetStarterContentFiles(TArray<FString>& OutFilenames);
 
@@ -409,6 +406,9 @@ private:
 
 	/** Given a source file name, find its location within the project */
 	static bool FindSourceFileInProject(const FString& InFilename, const FString& InSearchPath, FString& OutPath);
+	
+	/** Insert any required feature packs into the DefaultGame.ini file */
+	static bool InsertFeaturePacksIntoINIFile(const FProjectInformation& InProjectInfo, FText& OutFailReason);
 
 private:
 	static TWeakPtr<SNotificationItem> UpdateGameProjectNotification;
diff --git a/Engine/Source/Editor/GameProjectGeneration/Private/SNewProjectWizard.cpp b/Engine/Source/Editor/GameProjectGeneration/Private/SNewProjectWizard.cpp
--- a/Engine/Source/Editor/GameProjectGeneration/Private/SNewProjectWizard.cpp
+++ b/Engine/Source/Editor/GameProjectGeneration/Private/SNewProjectWizard.cpp
@@ -299,7 +299,8 @@ void SNewProjectWizard::Construct( const FArguments& InArgs )
 	bLastGlobalValidityCheckSuccessful = true;
 	bLastNameAndLocationValidityCheckSuccessful = true;
 	bPreventPeriodicValidityChecksUntilNextChange = false;
-	
+	bCopyStarterContent = true;
+
 	IHardwareTargetingModule& HardwareTargeting = IHardwareTargetingModule::Get();
 
 	SelectedHardwareClassTarget = EHardwareClass::Desktop;
diff --git a/Engine/Source/Editor/UnrealEd/Private/UnrealEdMisc.cpp b/Engine/Source/Editor/UnrealEd/Private/UnrealEdMisc.cpp
--- a/Engine/Source/Editor/UnrealEd/Private/UnrealEdMisc.cpp
+++ b/Engine/Source/Editor/UnrealEd/Private/UnrealEdMisc.cpp
@@ -256,6 +256,9 @@ void FUnrealEdMisc::OnInit()
 	{
 		bool bMapLoaded = false;
 
+		// Insert any feature packs if required. We need to do this before we try and load a map since any pack may contain a map
+		FeaturePackHandler.ImportPendingPacks();
+
 		FString ParsedMapName;
 		if ( FParse::Token(ParsedCmdLine, ParsedMapName, false) && 
 			 // If it's not a parameter
diff --git a/Engine/Source/Editor/UnrealEd/Public/UnrealEd.h b/Engine/Source/Editor/UnrealEd/Public/UnrealEd.h
--- a/Engine/Source/Editor/UnrealEd/Public/UnrealEd.h
+++ b/Engine/Source/Editor/UnrealEd/Public/UnrealEd.h
@@ -30,7 +30,7 @@
 
 #include "MRUList.h"
 #include "Projects.h"
-
+#include "FeaturePack/FeaturePack.h"
 
 //#include "../Private/GeomFitUtils.h"
 
diff --git a/Engine/Source/Editor/UnrealEd/Public/UnrealEdMisc.h b/Engine/Source/Editor/UnrealEd/Public/UnrealEdMisc.h
--- a/Engine/Source/Editor/UnrealEd/Public/UnrealEdMisc.h
+++ b/Engine/Source/Editor/UnrealEd/Public/UnrealEdMisc.h
@@ -320,4 +320,6 @@ private:
 
 	/** Package names and the number of times they have been updated */
 	TMap<FName, uint32> NumUpdatesByAssetName;
+
+	FFeaturePack FeaturePackHandler;
 };
diff --git a/Engine/Source/Runtime/Core/Private/Misc/Paths.cpp b/Engine/Source/Runtime/Core/Private/Misc/Paths.cpp
--- a/Engine/Source/Runtime/Core/Private/Misc/Paths.cpp
+++ b/Engine/Source/Runtime/Core/Private/Misc/Paths.cpp
@@ -352,9 +352,9 @@ FString FPaths::GameSourceDir()
 	return FPaths::GameDir() + TEXT("Source/");
 }
 
-FString FPaths::StarterContentDir()
+FString FPaths::FeaturePackDir()
 {
-	return FPaths::RootDir() + TEXT("Samples/StarterContent/");
+	return FPaths::RootDir() + TEXT("FeaturePacks/Packed/");
 }
 
 bool FPaths::IsProjectFilePathSet()
diff --git a/Engine/Source/Runtime/Core/Public/Misc/Paths.h b/Engine/Source/Runtime/Core/Public/Misc/Paths.h
--- a/Engine/Source/Runtime/Core/Public/Misc/Paths.h
+++ b/Engine/Source/Runtime/Core/Public/Misc/Paths.h
@@ -266,9 +266,9 @@ public:
 	static FString GameSourceDir();
 
 	/**
-	 * @return The directory where starter content is kept
+	 * @return The directory where feature packs are kept
 	 */
-	static FString StarterContentDir();
+	static FString FeaturePackDir();
 
 	/**
 	 * Checks whether the path to the project file, if any, is set.
