Reverted Starter Content change.

diff --git a/Engine/Source/Editor/AddContentDialog/Private/ContentSourceProviders/FeaturePack/FeaturePackContentSource.cpp b/Engine/Source/Editor/AddContentDialog/Private/ContentSourceProviders/FeaturePack/FeaturePackContentSource.cpp
--- a/Engine/Source/Editor/AddContentDialog/Private/ContentSourceProviders/FeaturePack/FeaturePackContentSource.cpp
+++ b/Engine/Source/Editor/AddContentDialog/Private/ContentSourceProviders/FeaturePack/FeaturePackContentSource.cpp
@@ -60,10 +60,6 @@ FFeaturePackContentSource::FFeaturePackContentSource(FString InFeaturePackPath)
 	{
 		Category = EContentSourceCategory::BlueprintFeature;
 	}
-	else if (CategoryString == "Content")
-	{
-		Category = EContentSourceCategory::Content;
-	}
 	else
 	{
 		Category = EContentSourceCategory::Unknown;
diff --git a/Engine/Source/Editor/AddContentDialog/Private/IContentSource.h b/Engine/Source/Editor/AddContentDialog/Private/IContentSource.h
--- a/Engine/Source/Editor/AddContentDialog/Private/IContentSource.h
+++ b/Engine/Source/Editor/AddContentDialog/Private/IContentSource.h
@@ -7,7 +7,6 @@ enum class EContentSourceCategory
 {
 	BlueprintFeature,
 	CodeFeature,
-	Content,
 	Unknown
 };
 
diff --git a/Engine/Source/Editor/AddContentDialog/Private/ViewModels/CategoryViewModel.cpp b/Engine/Source/Editor/AddContentDialog/Private/ViewModels/CategoryViewModel.cpp
--- a/Engine/Source/Editor/AddContentDialog/Private/ViewModels/CategoryViewModel.cpp
+++ b/Engine/Source/Editor/AddContentDialog/Private/ViewModels/CategoryViewModel.cpp
@@ -43,10 +43,6 @@ void FCategoryViewModel::Initialize()
 		Text = LOCTEXT("CodeFeature", "C++ Feature");
 		IconBrush = FAddContentDialogStyle::Get().GetBrush("AddContentDialog.CodeFeatureCategory");
 		break;
-	case EContentSourceCategory::Content:
-		Text = LOCTEXT("ContentPacks", "Content Packs");
-		IconBrush = FAddContentDialogStyle::Get().GetBrush("AddContentDialog.ContentPacks");
-		break;
 	default:
 		Text = LOCTEXT("Miscellaneous", "Miscellaneous");
 		IconBrush = FAddContentDialogStyle::Get().GetBrush("AddContentDialog.DefaultCategory");
diff --git a/Engine/Source/Editor/GameProjectGeneration/Classes/TemplateProjectDefs.h b/Engine/Source/Editor/GameProjectGeneration/Classes/TemplateProjectDefs.h
--- a/Engine/Source/Editor/GameProjectGeneration/Classes/TemplateProjectDefs.h
+++ b/Engine/Source/Editor/GameProjectGeneration/Classes/TemplateProjectDefs.h
@@ -96,9 +96,6 @@ class UTemplateProjectDefs : public UObject
 	UPROPERTY(config)
 	FString AssetTypes;
 
-	/* Optional list of feature packs to include */
-	UPROPERTY(config)
-	TArray<FString> PacksToInclude;
 
 	/** Fixes up all strings in this definitions object to replace \%TEMPLATENAME\% with the supplied template name and \%PROJECTNAME\% with the supplied project name */
 	void FixupStrings(const FString& TemplateName, const FString& ProjectName);
diff --git a/Engine/Source/Editor/GameProjectGeneration/GameProjectGeneration.Build.cs b/Engine/Source/Editor/GameProjectGeneration/GameProjectGeneration.Build.cs
--- a/Engine/Source/Editor/GameProjectGeneration/GameProjectGeneration.Build.cs
+++ b/Engine/Source/Editor/GameProjectGeneration/GameProjectGeneration.Build.cs
@@ -32,8 +32,7 @@ public class GameProjectGeneration : ModuleRules
  				"TargetPlatform",
 				"UnrealEd",
 				"DesktopPlatform",
-                "HardwareTargeting",                
-                "AssetTools",
+                "HardwareTargeting"
 			}
 		);
 
diff --git a/Engine/Source/Editor/GameProjectGeneration/Private/GameProjectUtils.cpp b/Engine/Source/Editor/GameProjectGeneration/Private/GameProjectUtils.cpp
--- a/Engine/Source/Editor/GameProjectGeneration/Private/GameProjectUtils.cpp
+++ b/Engine/Source/Editor/GameProjectGeneration/Private/GameProjectUtils.cpp
@@ -20,7 +20,7 @@
 #include "DesktopPlatformModule.h"
 #include "SNotificationList.h"
 #include "NotificationManager.h"
-#include "AssetToolsModule.h"
+#include "GameFramework/GameMode.h"
 
 #define LOCTEXT_NAMESPACE "GameProjectUtils"
 
@@ -505,11 +505,58 @@ bool GameProjectUtils::OpenCodeIDE(const FString& ProjectFile, FText& OutFailRea
 
 void GameProjectUtils::GetStarterContentFiles(TArray<FString>& OutFilenames)
 {
-	FString const SrcFolder = FPaths::FeaturePackDir();
-	
-	IFileManager::Get().FindFilesRecursive(OutFilenames, *SrcFolder, TEXT("*.upack"), /*Files=*/true, /*Directories=*/false);
+	FString const SrcFolder = FPaths::StarterContentDir();
+	FString const ContentFolder = SrcFolder / TEXT("Content");
+
+	// only copying /Content
+	IFileManager::Get().FindFilesRecursive(OutFilenames, *ContentFolder, TEXT("*"), /*Files=*/true, /*Directories=*/false);
 }
 
+bool GameProjectUtils::CopyStarterContent(const FString& DestProjectFolder, FText& OutFailReason)
+{
+	FString const SrcFolder = FPaths::StarterContentDir();
+
+	TArray<FString> FilesToCopy;
+	GetStarterContentFiles(FilesToCopy);
+
+	FScopedSlowTask SlowTask(FilesToCopy.Num(), LOCTEXT("CreatingProjectStatus_CopyingFiles", "Copying Files {SrcFilename}..."));
+	SlowTask.MakeDialog();
+
+	TArray<FString> CreatedFiles;
+	for (FString SrcFilename : FilesToCopy)
+	{
+		// Update the slow task dialog
+		FFormatNamedArguments Args;
+		Args.Add(TEXT("SrcFilename"), FText::FromString(FPaths::GetCleanFilename(SrcFilename)));
+		SlowTask.EnterProgressFrame(1, FText::Format(LOCTEXT("CreatingProjectStatus_CopyingFile", "Copying File {SrcFilename}..."), Args));
+
+		FString FileRelPath = FPaths::GetPath(SrcFilename);
+		FPaths::MakePathRelativeTo(FileRelPath, *SrcFolder);
+
+		// Perform the copy. For file collisions, leave existing file.
+		const FString DestFilename = DestProjectFolder + TEXT("/") + FileRelPath + TEXT("/") + FPaths::GetCleanFilename(SrcFilename);
+		if (!FPaths::FileExists(DestFilename))
+		{
+			if (IFileManager::Get().Copy(*DestFilename, *SrcFilename, false) == COPY_OK)
+			{
+				CreatedFiles.Add(DestFilename);
+			}
+			else
+			{
+				FFormatNamedArguments FailArgs;
+				FailArgs.Add(TEXT("SrcFilename"), FText::FromString(SrcFilename));
+				FailArgs.Add(TEXT("DestFilename"), FText::FromString(DestFilename));
+				OutFailReason = FText::Format(LOCTEXT("FailedToCopyFile", "Failed to copy \"{SrcFilename}\" to \"{DestFilename}\"."), FailArgs);
+				DeleteCreatedFiles(DestProjectFolder, CreatedFiles);
+				return false;
+			}
+		}
+	}
+
+	return true;
+}
+
+
 bool GameProjectUtils::CreateProject(const FProjectInformation& InProjectInfo, FText& OutFailReason)
 {
 	if ( !IsValidProjectFileForCreation(InProjectInfo.ProjectFilename, OutFailReason) )
@@ -805,9 +852,6 @@ bool GameProjectUtils::GenerateProjectFromScratch(const FProjectInformation& InP
 		return false;
 	}
 
-	// Insert any required feature packs (EG starter content) into ini file. These will be imported automatically when the editor is first run
-	InsertFeaturePacksIntoINIFile(InProjectInfo, OutFailReason);
-	
 	// Make the Content folder
 	const FString ContentFolder = NewProjectFolder / TEXT("Content");
 	if ( !IFileManager::Get().MakeDirectory(*ContentFolder) )
@@ -882,6 +926,17 @@ bool GameProjectUtils::GenerateProjectFromScratch(const FProjectInformation& InP
 
 	SlowTask.EnterProgressFrame();
 
+	if (InProjectInfo.bCopyStarterContent)
+	{
+		// Copy the starter content
+		if ( !CopyStarterContent(NewProjectFolder, OutFailReason) )
+		{
+			DeleteGeneratedProjectFiles(InProjectInfo.ProjectFilename);
+			DeleteCreatedFiles(NewProjectFolder, CreatedFiles);
+			return false;
+		}
+	}
+
 	UE_LOG(LogGameProjectGeneration, Log, TEXT("Created new project with %d files (plus project files)"), CreatedFiles.Num());
 	return true;
 }
@@ -1229,9 +1284,6 @@ bool GameProjectUtils::CreateProjectFromTemplate(const FProjectInformation& InPr
 		}
 	}
 
-	// Insert any required feature packs (EG starter content) into ini file. These will be imported automatically when the editor is first run
-	InsertFeaturePacksIntoINIFile(InProjectInfo, OutFailReason);
-	
 	SlowTask.EnterProgressFrame();
 
 	// Generate the project file
@@ -1299,6 +1351,17 @@ bool GameProjectUtils::CreateProjectFromTemplate(const FProjectInformation& InPr
 
 	SlowTask.EnterProgressFrame();
 
+	if (InProjectInfo.bCopyStarterContent)
+	{
+		// Copy the starter content
+		if ( !CopyStarterContent(DestFolder, OutFailReason) )
+		{
+			DeleteGeneratedProjectFiles(InProjectInfo.ProjectFilename);
+			DeleteCreatedFiles(DestFolder, CreatedFiles);
+			return false;
+		}
+	}
+
 	if (!TemplateDefs->PostGenerateProject(DestFolder, SrcFolder, InProjectInfo.ProjectFilename, InProjectInfo.TemplateFile, InProjectInfo.bShouldGenerateCode, OutFailReason))
 	{
 		DeleteGeneratedProjectFiles(InProjectInfo.ProjectFilename);
@@ -1480,23 +1543,66 @@ bool GameProjectUtils::GenerateConfigFiles(const FProjectInformation& InProjectI
 		
 		if (InProjectInfo.bCopyStarterContent)
 		{
+			FString StarterContentContentDir = FPaths::StarterContentDir() + TEXT("Content/");
+
+			TArray<FString> StarterContentMapFiles;
+			const FString FileWildcard = FString(TEXT("*")) + FPackageName::GetMapPackageExtension();
+		
 			FString SpecificEditorStartupMap;
-			FString SpecificGameDefaultMap;
+			FString SpecificGameDefaultMap;	
+			FString FullEditorStartupMapPath;
+			FString FullGameDefaultMapPath;
 
-			// If we have starter content packs available, specify starter map
-			if( IsStarterContentAvailableForNewProjects() == true )
+			// First we check if there are maps specified in the DefaultEngine.ini in our starter content folder			
+			const FString StarterContentDefaultEngineIniFilename = FPaths::StarterContentDir() / TEXT("Config/DefaultEngine.ini");
+			if (FPaths::FileExists(StarterContentDefaultEngineIniFilename))
 			{
-				if (InProjectInfo.TargetedHardware == EHardwareClass::Mobile)
+				FString StarterFileContents;
+				if (FFileHelper::LoadFileToString(StarterFileContents, *StarterContentDefaultEngineIniFilename))
 				{
-					SpecificEditorStartupMap = TEXT("/Game/MobileStarterContent/Maps/StarterMap");
-					SpecificGameDefaultMap = TEXT("/Game/MobileStarterContent/Maps/StarterMap");
-				}
-				else
-				{
-					SpecificEditorStartupMap = TEXT("/Game/StarterContent/Maps/StarterMap");
-					SpecificGameDefaultMap = TEXT("/Game/StarterContent/Maps/StarterMap");
+					TArray<FString> StarterIniLines;
+					StarterFileContents.ParseIntoArrayLines(&StarterIniLines);
+					for (int32 Line = 0; Line < StarterIniLines.Num();Line++)
+					{
+						FString EachLine = StarterIniLines[Line];
+						if (EachLine.StartsWith(TEXT("EditorStartupMap")))
+						{
+							EachLine.Split("=", nullptr, &SpecificEditorStartupMap);
+							FullEditorStartupMapPath = (StarterContentContentDir / SpecificEditorStartupMap) + FPackageName::GetMapPackageExtension();
+							FullEditorStartupMapPath = FullEditorStartupMapPath.Replace(TEXT("Game/"), TEXT(""));
+						}
+						if (EachLine.StartsWith(TEXT("GameDefaultMap")))
+						{
+							EachLine.Split("=", nullptr, &SpecificGameDefaultMap);
+							FullGameDefaultMapPath = (StarterContentContentDir / SpecificEditorStartupMap) + FPackageName::GetMapPackageExtension();
+							FullGameDefaultMapPath = FullGameDefaultMapPath.Replace(TEXT("Game/"), TEXT(""));
+						}
+					}					
 				}
-			}						
+			}
+
+			// Look for maps in the content folder. If we don't specify maps for EditorStartup and GameDefault we will use the first we find in here
+			IFileManager::Get().FindFilesRecursive(StarterContentMapFiles, *FPaths::StarterContentDir(), *FileWildcard, /*Files=*/true, /*Directories=*/false);			
+			FString MapPackagePath;
+			if (StarterContentMapFiles.Num() > 0)
+			{
+				const FString BaseMapFilename = FPaths::GetBaseFilename(StarterContentMapFiles[0]);
+
+				FString MapPathRelToContent = FPaths::GetPath(StarterContentMapFiles[0]);
+				FPaths::MakePathRelativeTo(MapPathRelToContent, *StarterContentContentDir);
+
+				MapPackagePath = FString(TEXT("/Game/")) + MapPathRelToContent + TEXT("/") + BaseMapFilename;
+			}
+
+			// if either the files we specified don't exist or we didn't specify any, use the first map file we found in the content folder.
+			if (SpecificEditorStartupMap.IsEmpty() || FPaths::FileExists(FullEditorStartupMapPath) == false)
+			{
+				SpecificEditorStartupMap = MapPackagePath;
+			}
+			if (SpecificGameDefaultMap.IsEmpty() || FPaths::FileExists(FullGameDefaultMapPath) == false)
+			{
+				SpecificGameDefaultMap = MapPackagePath;
+			}
 			
 			// Write out the settings for startup map and game default map
 			FileContents += TEXT("[/Script/EngineSettings.GameMapsSettings]") LINE_TERMINATOR;
@@ -1717,9 +1823,8 @@ bool GameProjectUtils::IsStarterContentAvailableForNewProjects()
 {
 	TArray<FString> StarterContentFiles;
 	GetStarterContentFiles(StarterContentFiles);
-	
-	bool bHasStaterContent = StarterContentFiles.FindByPredicate([&](const FString& Str){ return Str.Contains("StarterContent"); }) != nullptr;
-	return bHasStaterContent;
+
+	return (StarterContentFiles.Num() > 0);
 }
 
 TArray<FModuleContextInfo> GameProjectUtils::GetCurrentProjectModules()
@@ -2828,7 +2933,7 @@ bool GameProjectUtils::AddCodeToProject_Internal(const FString& NewClassName, co
 bool GameProjectUtils::FindSourceFileInProject(const FString& InFilename, const FString& InSearchPath, FString& OutPath)
 {
 	TArray<FString> Filenames;
-	const FString FilenameWidcard = /*TEXT("*") +*/ InFilename;
+	const FString FilenameWidcard = TEXT("*") + InFilename;
 	IFileManager::Get().FindFilesRecursive(Filenames, *InSearchPath, *FilenameWidcard, true, false, false);
 	
 	if(Filenames.Num())
@@ -2841,63 +2946,4 @@ bool GameProjectUtils::FindSourceFileInProject(const FString& InFilename, const
 	return false;
 }
 
-bool GameProjectUtils::InsertFeaturePacksIntoINIFile(const FProjectInformation& InProjectInfo, FText& OutFailReason)
-{
-	bool bSuccessfullyProcessed = false;
-	const FString ProjectName = FPaths::GetBaseFilename(InProjectInfo.ProjectFilename);
-	const FString TemplateName = FPaths::GetBaseFilename(InProjectInfo.TemplateFile);
-	const FString SrcFolder = FPaths::GetPath(InProjectInfo.TemplateFile);
-	const FString DestFolder = FPaths::GetPath(InProjectInfo.ProjectFilename);
-	
-	const FString ProjectConfigPath = DestFolder / TEXT("Config");
-	const FString IniFilename = ProjectConfigPath / TEXT("DefaultGame.ini");
-		
-	TArray<FString> PackList;
-	
-	// First the starter content
-	if (InProjectInfo.bCopyStarterContent)
-	{
-		FString StarterPack;
-		if( InProjectInfo.TargetedHardware == EHardwareClass::Mobile )
-		{
-			StarterPack = TEXT("InsertPack=(PackSource=\"MobileStarterContent.upack\",PackName=\"StarterContent\")");
-		}
-		else
-		{
-			StarterPack = TEXT("InsertPack=(PackSource=\"StarterContent.upack\",PackName=\"StarterContent\")");
-		}
-		PackList.Add(StarterPack);
-	}
-
-	// Now any packs specfied in the template def.
-	UTemplateProjectDefs* TemplateDefs = LoadTemplateDefs(SrcFolder);
-	if (TemplateDefs != NULL)
-	{
-		for (int32 iPack = 0; iPack < TemplateDefs->PacksToInclude.Num(); ++iPack)
-		{
-			PackList.Add(TemplateDefs->PacksToInclude[iPack]);
-		}
-	}
-
-	FString FileOutput;
-	if( PackList.Num() != 0)
-	{
-		FileOutput = TEXT("[StartupActions]");
-		FileOutput += LINE_TERMINATOR;
-		FileOutput += TEXT("bAddPacks=True");
-		FileOutput += LINE_TERMINATOR;
-		for (int32 iLine = 0; iLine < PackList.Num(); ++iLine)
-		{
-			FileOutput += PackList[iLine] + LINE_TERMINATOR;
-			
-		}
-		if (FFileHelper::SaveStringToFile(FileOutput, *IniFilename))
-		{
-			bSuccessfullyProcessed = true;
-		}
-	}
-	return bSuccessfullyProcessed;
-}
-
-
 #undef LOCTEXT_NAMESPACE
diff --git a/Engine/Source/Editor/GameProjectGeneration/Private/GameProjectUtils.h b/Engine/Source/Editor/GameProjectGeneration/Private/GameProjectUtils.h
--- a/Engine/Source/Editor/GameProjectGeneration/Private/GameProjectUtils.h
+++ b/Engine/Source/Editor/GameProjectGeneration/Private/GameProjectUtils.h
@@ -286,6 +286,9 @@ private:
 	/** Sets the engine association for a new project. Handles foreign and non-foreign projects. */
 	static bool SetEngineAssociationForForeignProject(const FString& ProjectFileName, FText& OutFailReason);
 
+	/** Copies starter content into the specified project folder. */
+	static bool CopyStarterContent(const FString& DestProjectFolder, FText& OutFailReason);
+
 	/** Returns list of starter content files */
 	static void GetStarterContentFiles(TArray<FString>& OutFilenames);
 
@@ -406,9 +409,6 @@ private:
 
 	/** Given a source file name, find its location within the project */
 	static bool FindSourceFileInProject(const FString& InFilename, const FString& InSearchPath, FString& OutPath);
-	
-	/** Insert any required feature packs into the DefaultGame.ini file */
-	static bool InsertFeaturePacksIntoINIFile(const FProjectInformation& InProjectInfo, FText& OutFailReason);
 
 private:
 	static TWeakPtr<SNotificationItem> UpdateGameProjectNotification;
diff --git a/Engine/Source/Editor/UnrealEd/Private/FeaturePack/FeaturePack.cpp b/Engine/Source/Editor/UnrealEd/Private/FeaturePack/FeaturePack.cpp
deleted file mode 100644
diff --git a/Engine/Source/Editor/UnrealEd/Private/UnrealEdMisc.cpp b/Engine/Source/Editor/UnrealEd/Private/UnrealEdMisc.cpp
--- a/Engine/Source/Editor/UnrealEd/Private/UnrealEdMisc.cpp
+++ b/Engine/Source/Editor/UnrealEd/Private/UnrealEdMisc.cpp
@@ -256,9 +256,6 @@ void FUnrealEdMisc::OnInit()
 	{
 		bool bMapLoaded = false;
 
-		// Insert any feature packs if required. We need to do this before we try and load a map since any pack may contain a map
-		FeaturePackHandler.ImportPendingPacks();
-
 		FString ParsedMapName;
 		if ( FParse::Token(ParsedCmdLine, ParsedMapName, false) && 
 			 // If it's not a parameter
diff --git a/Engine/Source/Editor/UnrealEd/Public/FeaturePack/FeaturePack.h b/Engine/Source/Editor/UnrealEd/Public/FeaturePack/FeaturePack.h
deleted file mode 100644
diff --git a/Engine/Source/Editor/UnrealEd/Public/UnrealEd.h b/Engine/Source/Editor/UnrealEd/Public/UnrealEd.h
--- a/Engine/Source/Editor/UnrealEd/Public/UnrealEd.h
+++ b/Engine/Source/Editor/UnrealEd/Public/UnrealEd.h
@@ -30,7 +30,7 @@
 
 #include "MRUList.h"
 #include "Projects.h"
-#include "FeaturePack/FeaturePack.h"
+
 
 //#include "../Private/GeomFitUtils.h"
 
diff --git a/Engine/Source/Editor/UnrealEd/Public/UnrealEdMisc.h b/Engine/Source/Editor/UnrealEd/Public/UnrealEdMisc.h
--- a/Engine/Source/Editor/UnrealEd/Public/UnrealEdMisc.h
+++ b/Engine/Source/Editor/UnrealEd/Public/UnrealEdMisc.h
@@ -320,6 +320,4 @@ private:
 
 	/** Package names and the number of times they have been updated */
 	TMap<FName, uint32> NumUpdatesByAssetName;
-
-	FFeaturePack FeaturePackHandler;
 };
diff --git a/Engine/Source/Runtime/Core/Private/Misc/Paths.cpp b/Engine/Source/Runtime/Core/Private/Misc/Paths.cpp
--- a/Engine/Source/Runtime/Core/Private/Misc/Paths.cpp
+++ b/Engine/Source/Runtime/Core/Private/Misc/Paths.cpp
@@ -352,9 +352,9 @@ FString FPaths::GameSourceDir()
 	return FPaths::GameDir() + TEXT("Source/");
 }
 
-FString FPaths::FeaturePackDir()
+FString FPaths::StarterContentDir()
 {
-	return FPaths::RootDir() + TEXT("FeaturePacks/Packed/");
+	return FPaths::RootDir() + TEXT("Samples/StarterContent/");
 }
 
 bool FPaths::IsProjectFilePathSet()
@@ -554,13 +554,13 @@ bool FPaths::IsDrive(const FString& InPath)
 bool FPaths::IsRelative(const FString& InPath)
 {
 	// The previous implementation of this function seemed to handle normalized and unnormalized paths, so this one does too for legacy reasons.
+
 	const bool IsRooted =	InPath.StartsWith(TEXT("\\\\"))	||												// "\\" for UNC or "network" paths.
 							InPath.StartsWith(TEXT("//"))	||												// Equivalent to "\\", considering normalization replaces "\\" with "//".
 							InPath.StartsWith(TEXT("\\"))	||												// Root of the current directory on Windows
 							InPath.StartsWith(TEXT("/"))	||												// Root of the current directory on Windows, root on UNIX-likes.
-							InPath.StartsWith(TEXT("root:/")) ||											// Feature packs use this
 							(InPath.Len() >= 2 && FChar::IsAlpha(InPath[0]) && InPath[1] == TEXT(':'));	// Starts with "<DriveLetter>:"
-	 
+
 	return !IsRooted;
 }
 
diff --git a/Engine/Source/Runtime/Core/Public/Misc/Paths.h b/Engine/Source/Runtime/Core/Public/Misc/Paths.h
--- a/Engine/Source/Runtime/Core/Public/Misc/Paths.h
+++ b/Engine/Source/Runtime/Core/Public/Misc/Paths.h
@@ -266,9 +266,9 @@ public:
 	static FString GameSourceDir();
 
 	/**
-	 * @return The directory where feature packs are kept
+	 * @return The directory where starter content is kept
 	 */
-	static FString FeaturePackDir();
+	static FString StarterContentDir();
 
 	/**
 	 * Checks whether the path to the project file, if any, is set.
@@ -315,7 +315,7 @@ public:
 	/** @return true if this path represents a drive */
 	static bool IsDrive(const FString& InPath);
 
-	/** @return true if this path is relative (not rooted) */
+	/** @return true if this path is relative */
 	static bool IsRelative(const FString& InPath);
 
 	/** Convert all / and \ to TEXT("/") */
