NewProjectWizard - Fix text flicker on tab switch.

diff --git a/Engine/Source/Editor/GameProjectGeneration/Private/SNewProjectWizard.cpp b/Engine/Source/Editor/GameProjectGeneration/Private/SNewProjectWizard.cpp
--- a/Engine/Source/Editor/GameProjectGeneration/Private/SNewProjectWizard.cpp
+++ b/Engine/Source/Editor/GameProjectGeneration/Private/SNewProjectWizard.cpp
@@ -265,7 +265,7 @@ public:
 				.Padding(FMargin(ThumbnailPadding, 0))
 				[
 					SNew(STextBlock)
-					.AutoWrapText(true)
+					.WrapTextAt(ThumbnailSize)
 					.Justification(ETextJustify::Center)
 					.LineBreakPolicy(FBreakIterator::CreateCamelCaseBreakIterator())
 					//.HighlightText(this, &SNewProjectWizard::GetItemHighlightText)
