Immediately compile and open new code projects (rather than opening the IDE and requiring the user do it themselves).

diff --git a/Engine/Source/Editor/GameProjectGeneration/Private/GameProjectUtils.cpp b/Engine/Source/Editor/GameProjectGeneration/Private/GameProjectUtils.cpp
--- a/Engine/Source/Editor/GameProjectGeneration/Private/GameProjectUtils.cpp
+++ b/Engine/Source/Editor/GameProjectGeneration/Private/GameProjectUtils.cpp
@@ -22,6 +22,7 @@
 #include "NotificationManager.h"
 #include "GameFramework/GameMode.h"
 #include "HotReloadInterface.h"
+#include "SVerbChoiceDialog.h"
 
 #define LOCTEXT_NAMESPACE "GameProjectUtils"
 
@@ -1801,6 +1802,41 @@ bool GameProjectUtils::GenerateGameFrameworkSourceCode(const FString& NewProject
 	return true;
 }
 
+bool GameProjectUtils::BuildCodeProject(const FString& ProjectFilename)
+{
+	// Build the project while capturing the log output. Passing GWarn to CompileGameProject will allow Slate to display the progress bar.
+	FStringOutputDevice OutputLog;
+	OutputLog.SetAutoEmitLineTerminator(true);
+	GLog->AddOutputDevice(&OutputLog);
+	bool bCompileSucceeded = FDesktopPlatformModule::Get()->CompileGameProject(FPaths::RootDir(), ProjectFilename, GWarn);
+	GLog->RemoveOutputDevice(&OutputLog);
+
+	// Try to compile the modules
+	if(!bCompileSucceeded)
+	{
+		FText DevEnvName = FSourceCodeNavigation::GetSuggestedSourceCodeIDE( true );
+
+		TArray<FText> CompileFailedButtons;
+		int32 OpenIDEButton = CompileFailedButtons.Add(FText::Format(LOCTEXT("CompileFailedOpenIDE", "Open with {0}"), DevEnvName));
+		int32 ViewLogButton = CompileFailedButtons.Add(LOCTEXT("CompileFailedViewLog", "View build log"));
+		CompileFailedButtons.Add(LOCTEXT("CompileFailedCancel", "Cancel"));
+
+		int32 CompileFailedChoice = SVerbChoiceDialog::ShowModal(LOCTEXT("ProjectUpgradeTitle", "Project Conversion Failed"), FText::Format(LOCTEXT("ProjectUpgradeCompileFailed", "The project failed to compile with this version of the engine. Would you like to open the project in {0}?"), DevEnvName), CompileFailedButtons);
+		if(CompileFailedChoice == ViewLogButton)
+		{
+			CompileFailedButtons.RemoveAt(ViewLogButton);
+			CompileFailedChoice = SVerbChoiceDialog::ShowModal(LOCTEXT("ProjectUpgradeTitle", "Project Conversion Failed"), FText::Format(LOCTEXT("ProjectUpgradeCompileFailed", "The project failed to compile with this version of the engine. Build output is as follows:\n\n{0}"), FText::FromString(OutputLog)), CompileFailedButtons);
+		}
+
+		FText FailReason;
+		if(CompileFailedChoice == OpenIDEButton && !GameProjectUtils::OpenCodeIDE(ProjectFilename, FailReason))
+		{
+			FMessageDialog::Open(EAppMsgType::Ok, FailReason);
+		}
+	}
+	return bCompileSucceeded;
+}
+
 bool GameProjectUtils::GenerateCodeProjectFiles(const FString& ProjectFilename, FText& OutFailReason)
 {
 	FStringOutputDevice OutputLog;
diff --git a/Engine/Source/Editor/GameProjectGeneration/Private/GameProjectUtils.h b/Engine/Source/Editor/GameProjectGeneration/Private/GameProjectUtils.h
--- a/Engine/Source/Editor/GameProjectGeneration/Private/GameProjectUtils.h
+++ b/Engine/Source/Editor/GameProjectGeneration/Private/GameProjectUtils.h
@@ -197,6 +197,9 @@ public:
 	/** Returns the uproject template filename for the default project template. */
 	static FString GetDefaultProjectTemplateFilename();
 
+	/** Compiles a project while showing a progress bar, and offers to open the IDE if it fails. */
+	static bool BuildCodeProject(const FString& ProjectFilename);
+
 	/** Creates code project files for a new game project. On failure, OutFailReason will be populated. */
 	static bool GenerateCodeProjectFiles(const FString& ProjectFilename, FText& OutFailReason);
 
diff --git a/Engine/Source/Editor/GameProjectGeneration/Private/SNewProjectWizard.cpp b/Engine/Source/Editor/GameProjectGeneration/Private/SNewProjectWizard.cpp
--- a/Engine/Source/Editor/GameProjectGeneration/Private/SNewProjectWizard.cpp
+++ b/Engine/Source/Editor/GameProjectGeneration/Private/SNewProjectWizard.cpp
@@ -1462,39 +1462,18 @@ void SNewProjectWizard::CreateAndOpenProject( )
 			// Prevent periodic validity checks. This is to prevent a brief error message about the project already existing while you are exiting.
 			bPreventPeriodicValidityChecksUntilNextChange = true;
 
-			const bool bCodeAdded = GetSelectedTemplateItem()->bGenerateCode;
-			if ( bCodeAdded )
+			// If it's a code project, compile the binaries first.
+			if(!GetSelectedTemplateItem()->bGenerateCode || GameProjectUtils::BuildCodeProject(ProjectFile))
 			{
-				// In non-rocket, the engine executable may need to be built in order to build the game binaries,
-				// just open the code editing ide now instead of automatically building for them since it is not safe to do so.
-				const bool bPromptForConfirmation = FApp::HasGameName(); /** Only prompt for project switching if we are already in a project */
-				OpenCodeIDE( ProjectFile, bPromptForConfirmation );
-			}
-			else
-			{
-				// Successfully created a content only project. Now open it.
-				const bool bPromptForConfirmation = false;
-				OpenProject( ProjectFile, bPromptForConfirmation );
+				OpenProject( ProjectFile );
 			}
 		}
 	}
 }
 
 
-bool SNewProjectWizard::OpenProject( const FString& ProjectFile, bool bPromptForConfirmation )
+bool SNewProjectWizard::OpenProject( const FString& ProjectFile )
 {
-	if ( bPromptForConfirmation )
-	{
-		// Notify the user of the success, and ask to switch projects.
-		FText SuccessMessage = FText::Format( LOCTEXT("NewProjectSuccessful", "Project '{0}' was successfully created. Would you like to open it now?"), FText::FromString(FPaths::GetBaseFilename(ProjectFile)) );
-		if ( FMessageDialog::Open( EAppMsgType::YesNo, SuccessMessage ) == EAppReturnType::No )
-		{
-			// The user opted out of opening the new project. Just close the window.
-			CloseWindowIfAppropriate();
-			return false;
-		}
-	}
-
 	FText FailReason;
 	if ( GameProjectUtils::OpenProject( ProjectFile, FailReason ) )
 	{
@@ -1509,20 +1488,8 @@ bool SNewProjectWizard::OpenProject( const FString& ProjectFile, bool bPromptFor
 }
 
 
-bool SNewProjectWizard::OpenCodeIDE( const FString& ProjectFile, bool bPromptForConfirmation )
+bool SNewProjectWizard::OpenCodeIDE( const FString& ProjectFile )
 {
-	if ( bPromptForConfirmation )
-	{
-		// Notify the user of the success, and ask to switch projects.
-		FText SuccessMessage = FText::Format( LOCTEXT("NewProjectSuccessfulIDE", "Project '{0}' was successfully created. Would you like to open it in {1}?"), FText::FromString(FPaths::GetBaseFilename(ProjectFile)), FSourceCodeNavigation::GetSuggestedSourceCodeIDE() );
-		if ( FMessageDialog::Open( EAppMsgType::YesNo, SuccessMessage ) == EAppReturnType::No )
-		{
-			// The user opted out of opening the new project. Just close the window.
-			CloseWindowIfAppropriate(true);
-			return false;
-		}
-	}
-
 	FText FailReason;
 
 	if ( GameProjectUtils::OpenCodeIDE( ProjectFile, FailReason ) )
diff --git a/Engine/Source/Editor/GameProjectGeneration/Private/SNewProjectWizard.h b/Engine/Source/Editor/GameProjectGeneration/Private/SNewProjectWizard.h
--- a/Engine/Source/Editor/GameProjectGeneration/Private/SNewProjectWizard.h
+++ b/Engine/Source/Editor/GameProjectGeneration/Private/SNewProjectWizard.h
@@ -163,10 +163,10 @@ private:
 	void CreateAndOpenProject();
 
 	/** Opens the specified project file */
-	bool OpenProject(const FString& ProjectFile, bool bPromptForConfirmation);
+	bool OpenProject(const FString& ProjectFile);
 
 	/** Opens the solution for the specified project */
-	bool OpenCodeIDE(const FString& ProjectFile, bool bPromptForConfirmation);
+	bool OpenCodeIDE(const FString& ProjectFile);
 
 	/** Creates a project with the supplied project filename */
 	bool CreateProject(const FString& ProjectFile);
diff --git a/Engine/Source/Editor/GameProjectGeneration/Private/SProjectBrowser.cpp b/Engine/Source/Editor/GameProjectGeneration/Private/SProjectBrowser.cpp
--- a/Engine/Source/Editor/GameProjectGeneration/Private/SProjectBrowser.cpp
+++ b/Engine/Source/Editor/GameProjectGeneration/Private/SProjectBrowser.cpp
@@ -1047,33 +1047,8 @@ bool SProjectBrowser::OpenProject( const FString& InProjectFile )
 				}
 
 				// Try to compile the project
-				OutputLog.Empty();
-				GLog->AddOutputDevice(&OutputLog);
-				bool bCompileSucceeded = FDesktopPlatformModule::Get()->CompileGameProject(FPaths::RootDir(), ProjectFile, GWarn);
-				GLog->RemoveOutputDevice(&OutputLog);
-
-				// Try to compile the modules
-				if(!bCompileSucceeded)
+				if(!GameProjectUtils::BuildCodeProject(ProjectFile))
 				{
-					FText DevEnvName = FSourceCodeNavigation::GetSuggestedSourceCodeIDE( true );
-
-					TArray<FText> CompileFailedButtons;
-					int32 OpenIDEButton = CompileFailedButtons.Add(FText::Format(LOCTEXT("CompileFailedOpenIDE", "Open with {0}"), DevEnvName));
-					int32 ViewLogButton = CompileFailedButtons.Add(LOCTEXT("CompileFailedViewLog", "View build log"));
-					CompileFailedButtons.Add(LOCTEXT("CompileFailedCancel", "Cancel"));
-
-					int32 CompileFailedChoice = SVerbChoiceDialog::ShowModal(LOCTEXT("ProjectUpgradeTitle", "Project Conversion Failed"), FText::Format(LOCTEXT("ProjectUpgradeCompileFailed", "The project failed to compile with this version of the engine. Would you like to open the project in {0}?"), DevEnvName), CompileFailedButtons);
-					if(CompileFailedChoice == ViewLogButton)
-					{
-						CompileFailedButtons.RemoveAt(ViewLogButton);
-						CompileFailedChoice = SVerbChoiceDialog::ShowModal(LOCTEXT("ProjectUpgradeTitle", "Project Conversion Failed"), FText::Format(LOCTEXT("ProjectUpgradeCompileFailed", "The project failed to compile with this version of the engine. Build output is as follows:\n\n{0}"), FText::FromString(OutputLog)), CompileFailedButtons);
-					}
-
-					if(CompileFailedChoice == OpenIDEButton && !GameProjectUtils::OpenCodeIDE(ProjectFile, FailReason))
-					{
-						FMessageDialog::Open(EAppMsgType::Ok, FailReason);
-					}
-
 					return false;
 				}
 			}
