Fix non-unity in GameProjectUtils.cpp

diff --git a/Engine/Source/Editor/GameProjectGeneration/Private/GameProjectUtils.cpp b/Engine/Source/Editor/GameProjectGeneration/Private/GameProjectUtils.cpp
--- a/Engine/Source/Editor/GameProjectGeneration/Private/GameProjectUtils.cpp
+++ b/Engine/Source/Editor/GameProjectGeneration/Private/GameProjectUtils.cpp
@@ -15,6 +15,7 @@
 #include "TargetPlatform.h"
 
 #include "ClassIconFinder.h"
+#include "Editor/UnrealEd/Public/SourceCodeNavigation.h"
 
 #include "UProjectInfo.h"
 #include "DesktopPlatformModule.h"
@@ -24,6 +25,7 @@
 #include "HotReloadInterface.h"
 #include "SVerbChoiceDialog.h"
 
+
 #define LOCTEXT_NAMESPACE "GameProjectUtils"
 
 #define MAX_PROJECT_PATH_BUFFER_SPACE 130 // Leave a reasonable buffer of additional characters to account for files created in the content directory during or after project generation
