Editor: Add some verbose logging to help diagonse issues when creating new project templates

diff --git a/Engine/Source/Editor/GameProjectGeneration/Private/GameProjectUtils.cpp b/Engine/Source/Editor/GameProjectGeneration/Private/GameProjectUtils.cpp
--- a/Engine/Source/Editor/GameProjectGeneration/Private/GameProjectUtils.cpp
+++ b/Engine/Source/Editor/GameProjectGeneration/Private/GameProjectUtils.cpp
@@ -1010,7 +1010,7 @@ bool GameProjectUtils::CreateProjectFromTemplate(const FProjectInformation& InPr
 			Args.Add( TEXT("SrcFilename"), FText::FromString( FPaths::GetCleanFilename(SrcFilename) ) );
 			InnerSlowTask.EnterProgressFrame(1, FText::Format( LOCTEXT( "CreatingProjectStatus_CopyingFile", "Copying File {SrcFilename}..." ), Args ));
 
-		// Get the file path, relative to the src folder
+			// Get the file path, relative to the src folder
 			const FString SrcFileSubpath = SrcFilename.RightChop(SrcFolder.Len() + 1);
 
 			// Skip any files that were configured to be ignored
@@ -1038,6 +1038,7 @@ bool GameProjectUtils::CreateProjectFromTemplate(const FProjectInformation& InPr
 				if ( SrcFileSubpath.StartsWith((*IgnoreIt) + TEXT("/") ) )
 				{
 					// This folder was marked as "ignored"
+					UE_LOG(LogGameProjectGeneration, Verbose, TEXT("'%s': Skipping as it is in an ignored folder '%s'"), *SrcFilename, **IgnoreIt);
 					bThisFolderIsIgnored = true;
 					break;
 				}
@@ -1058,6 +1059,7 @@ bool GameProjectUtils::CreateProjectFromTemplate(const FProjectInformation& InPr
 				{
 					// This was a file in a renamed folder. Retarget to the new location
 					DestFileSubpathWithoutFilename = FolderRename.To / DestFileSubpathWithoutFilename.RightChop( FolderRename.From.Len() );
+					UE_LOG(LogGameProjectGeneration, Verbose, TEXT("'%s': Moving to '%s' as it matched folder rename ('%s'->'%s')"), *SrcFilename, *DestFileSubpathWithoutFilename, *FolderRename.From, *FolderRename.To);
 				}
 			}
 
@@ -1070,7 +1072,13 @@ bool GameProjectUtils::CreateProjectFromTemplate(const FProjectInformation& InPr
 				if ( Replacement.Extensions.Contains( FileExtension ) )
 				{
 					// This file matched a filename replacement extension, apply it now
+					FString LastDestBaseFilename = DestBaseFilename;
 					DestBaseFilename = DestBaseFilename.Replace(*Replacement.From, *Replacement.To, Replacement.bCaseSensitive ? ESearchCase::CaseSensitive : ESearchCase::IgnoreCase);
+
+					if (LastDestBaseFilename != DestBaseFilename)
+					{
+						UE_LOG(LogGameProjectGeneration, Verbose, TEXT("'%s': Renaming to '%s/%s' as it matched file rename ('%s'->'%s')"), *SrcFilename, *DestFileSubpathWithoutFilename, *DestBaseFilename, *Replacement.From, *Replacement.To);
+					}
 				}
 			}
 
