Changed source folder thumb/screen shot used by templates so it can share the asset with feature packs
revised code pack insert (now replaces feature pack header name on import instead of using %MODULE_INCLUDE_PATH%

diff --git a/Engine/Source/Editor/GameProjectGeneration/Private/SNewProjectWizard.cpp b/Engine/Source/Editor/GameProjectGeneration/Private/SNewProjectWizard.cpp
--- a/Engine/Source/Editor/GameProjectGeneration/Private/SNewProjectWizard.cpp
+++ b/Engine/Source/Editor/GameProjectGeneration/Private/SNewProjectWizard.cpp
@@ -1211,7 +1211,7 @@ void SNewProjectWizard::FindTemplateProjects()
 					const bool bGenerateCode = TemplateDefs->GeneratesCode(Root);
 
 					TSharedPtr<FSlateDynamicImageBrush> ThumbnailBrush;
-					const FString ThumbnailPNGFile = FPaths::GetBaseFilename(ProjectFilename, false) + TEXT(".png");
+					const FString ThumbnailPNGFile = (Root + TEXT("/Media/") + FoundProjectFiles[0]).Replace(TEXT(".uproject"), TEXT(".png"));
 					if ( FPlatformFileManager::Get().GetPlatformFile().FileExists(*ThumbnailPNGFile) )
 					{
 						const FName BrushName = FName(*ThumbnailPNGFile);
@@ -1219,7 +1219,7 @@ void SNewProjectWizard::FindTemplateProjects()
 					}
 
 					TSharedPtr<FSlateDynamicImageBrush> PreviewBrush;
-					const FString PreviewPNGFile = FPaths::GetBaseFilename(ProjectFilename, false) + TEXT("_Preview.png");
+					const FString PreviewPNGFile = (Root + TEXT("/Media/") + FoundProjectFiles[0]).Replace(TEXT(".uproject"), TEXT("_Preview.png"));
 					if ( FPlatformFileManager::Get().GetPlatformFile().FileExists(*PreviewPNGFile) )
 					{
 						const FName BrushName = FName(*PreviewPNGFile);
diff --git a/Engine/Source/Editor/UnrealEd/Private/Factories/PackFactory.cpp b/Engine/Source/Editor/UnrealEd/Private/Factories/PackFactory.cpp
--- a/Engine/Source/Editor/UnrealEd/Private/Factories/PackFactory.cpp
+++ b/Engine/Source/Editor/UnrealEd/Private/Factories/PackFactory.cpp
@@ -419,7 +419,9 @@ UObject* UPackFactory::FactoryCreateBinary
 					PackFactoryHelper::ExtractFileToString(Entry, PakReader, Buffer, PersistentCompressionBuffer, SourceContents);
 
 					FGameProjectGenerationModule& GameProjectModule = FModuleManager::LoadModuleChecked<FGameProjectGenerationModule>(TEXT("GameProjectGeneration"));
-					SourceContents = SourceContents.Replace(TEXT("%MODULE_INCLUDE_PATH%"), *GameProjectModule.Get().DetermineModuleIncludePath(SourceModuleInfo, DestFilename), ESearchCase::CaseSensitive);
+					FString StringToReplace = *ConfigParameters.Find("GameName");
+					StringToReplace += ".h";
+					SourceContents = SourceContents.Replace(*StringToReplace, *GameProjectModule.Get().DetermineModuleIncludePath(SourceModuleInfo, DestFilename), ESearchCase::CaseSensitive);
 
 					if (FFileHelper::SaveStringToFile(SourceContents, *DestFilename))
 					{
