Fixed infinitely duplicating active tick is SGameProjectDialog - widgets fade in properly again

diff --git a/Engine/Source/Editor/GameProjectGeneration/Private/SGameProjectDialog.cpp b/Engine/Source/Editor/GameProjectGeneration/Private/SGameProjectDialog.cpp
--- a/Engine/Source/Editor/GameProjectGeneration/Private/SGameProjectDialog.cpp
+++ b/Engine/Source/Editor/GameProjectGeneration/Private/SGameProjectDialog.cpp
@@ -196,7 +196,7 @@ EActiveTickReturnType SGameProjectDialog::TriggerFadeInPostConstruct(double InCu
 	// see the dialog and the beginning of the animation is cut off.
 	FadeAnimation.Play(this->AsShared());
 
-	return EActiveTickReturnType::KeepTicking;
+	return EActiveTickReturnType::StopTicking;
 }
 
 /* SGameProjectDialog implementation
