Remove GenerateNewProjectFile from IProjectManager.

diff --git a/Engine/Source/Editor/GameProjectGeneration/Private/GameProjectUtils.cpp b/Engine/Source/Editor/GameProjectGeneration/Private/GameProjectUtils.cpp
--- a/Engine/Source/Editor/GameProjectGeneration/Private/GameProjectUtils.cpp
+++ b/Engine/Source/Editor/GameProjectGeneration/Private/GameProjectUtils.cpp
@@ -896,17 +896,22 @@ bool GameProjectUtils::GenerateProjectFromScratch(const FProjectInformation& InP
 
 	// Generate the project file
 	{
-		FText LocalFailReason;
-		if (IProjectManager::Get().GenerateNewProjectFile(InProjectInfo.ProjectFilename, StartupModuleNames, TEXT(""), LocalFailReason))
+		// Set up the descriptor
+		FProjectDescriptor Descriptor;
+		for(int32 Idx = 0; Idx < StartupModuleNames.Num(); Idx++)
 		{
-			CreatedFiles.Add(InProjectInfo.ProjectFilename);
+			Descriptor.Modules.Add(FModuleDescriptor(*StartupModuleNames[Idx]));
 		}
-		else
+
+		// Try to save it
+		FText LocalFailReason;
+		if(!Descriptor.Save(InProjectInfo.ProjectFilename, LocalFailReason))
 		{
 			OutFailReason = LocalFailReason;
 			DeleteCreatedFiles(NewProjectFolder, CreatedFiles);
 			return false;
 		}
+		CreatedFiles.Add(InProjectInfo.ProjectFilename);
 
 		// Set the engine identifier for it. Do this after saving, so it can be correctly detected as foreign or non-foreign.
 		if(!SetEngineAssociationForForeignProject(InProjectInfo.ProjectFilename, OutFailReason))
@@ -2757,9 +2762,6 @@ bool GameProjectUtils::UpdateGameProjectFile(const FString& ProjectFile, const F
 	FProjectDescriptor Descriptor;
 	if(Descriptor.Load(ProjectFile, OutFailReason))
 	{
-		// Freshen version information
-		Descriptor.EngineAssociation = EngineIdentifier;
-
 		// Replace the modules names, if specified
 		if(StartupModuleNames != NULL)
 		{
@@ -2771,7 +2773,7 @@ bool GameProjectUtils::UpdateGameProjectFile(const FString& ProjectFile, const F
 		}
 
 		// Update file on disk
-		return Descriptor.Save(ProjectFile, OutFailReason);
+		return Descriptor.Save(ProjectFile, OutFailReason) && FDesktopPlatformModule::Get()->SetEngineIdentifierForProject(ProjectFile, EngineIdentifier);
 	}
 	return false;
 }
diff --git a/Engine/Source/Runtime/Projects/Private/ProjectManager.cpp b/Engine/Source/Runtime/Projects/Private/ProjectManager.cpp
--- a/Engine/Source/Runtime/Projects/Private/ProjectManager.cpp
+++ b/Engine/Source/Runtime/Projects/Private/ProjectManager.cpp
@@ -107,19 +107,6 @@ const FString& FProjectManager::GetAutoLoadProjectFileName()
 	return RecentProjectFileName;
 }
 
-bool FProjectManager::GenerateNewProjectFile(const FString& NewProjectFilename, const TArray<FString>& StartupModuleNames, const FString& EngineIdentifier, FText& OutFailReason)
-{
-	FProjectDescriptor Descriptor;
-	Descriptor.EngineAssociation = EngineIdentifier;
-
-	for(int32 Idx = 0; Idx < StartupModuleNames.Num(); Idx++)
-	{
-		Descriptor.Modules.Add(FModuleDescriptor(*StartupModuleNames[Idx]));
-	}
-
-	return Descriptor.Save(NewProjectFilename, OutFailReason);
-}
-
 bool FProjectManager::SignSampleProject(const FString& FilePath, const FString& Category, FText& OutFailReason)
 {
 	FProjectDescriptor Descriptor;
diff --git a/Engine/Source/Runtime/Projects/Private/ProjectManager.h b/Engine/Source/Runtime/Projects/Private/ProjectManager.h
--- a/Engine/Source/Runtime/Projects/Private/ProjectManager.h
+++ b/Engine/Source/Runtime/Projects/Private/ProjectManager.h
@@ -18,7 +18,6 @@ public:
 	virtual bool LoadModulesForProject( const ELoadingPhase::Type LoadingPhase ) override;
 	virtual bool CheckModuleCompatibility( TArray<FString>& OutIncompatibleModules ) override;
 	virtual const FString& GetAutoLoadProjectFileName() override;
-	virtual bool GenerateNewProjectFile(const FString& NewProjectFilename, const TArray<FString>& ModuleNames, const FString& EngineIdentifier, FText& OutFailReason) override;
 	virtual bool SignSampleProject(const FString& FilePath, const FString& Category, FText& OutFailReason) override;
 	virtual bool QueryStatusForProject(const FString& FilePath, FProjectStatus& OutProjectStatus) const override;
 	virtual bool QueryStatusForCurrentProject(FProjectStatus& OutProjectStatus) const override;
diff --git a/Engine/Source/Runtime/Projects/Public/Interfaces/IProjectManager.h b/Engine/Source/Runtime/Projects/Public/Interfaces/IProjectManager.h
--- a/Engine/Source/Runtime/Projects/Public/Interfaces/IProjectManager.h
+++ b/Engine/Source/Runtime/Projects/Public/Interfaces/IProjectManager.h
@@ -119,18 +119,6 @@ public:
 	 */
 	virtual const FString& GetAutoLoadProjectFileName() = 0;
 
-	/**
-	 * Generates a new project file and saves it to disk at the specified location
-	 * 
-	 * @param NewProjectFilename	The filename of the file to be written out
-	 * @param StartupModuleNames	The list of modules to be loaded at startup
-	 * @param EngineIdentifier		Identifier of the engine that the project should be associated with
-	 * @param OutFailReason			When returning false, this provides a display reason why the file could not be created.
-	 *
-	 * @return true when the file was successfully written
-	 */
-	virtual bool GenerateNewProjectFile(const FString& NewProjectFilename, const TArray<FString>& StartupModuleNames, const FString& EngineIdentifier, FText& OutFailReason) = 0;
-
 	/**
 	 * Sets the project's EpicSampleNameHash (based on its filename) and category, then saves the file to disk.
 	 * This marks the project as a sample and fixes its filename so that
diff --git a/Engine/Source/Runtime/Projects/Public/ProjectDescriptor.h b/Engine/Source/Runtime/Projects/Public/ProjectDescriptor.h
--- a/Engine/Source/Runtime/Projects/Public/ProjectDescriptor.h
+++ b/Engine/Source/Runtime/Projects/Public/ProjectDescriptor.h
@@ -32,7 +32,7 @@ struct PROJECTS_API FProjectDescriptor
 	/** Descriptor version number. */
 	EProjectDescriptorVersion::Type FileVersion;
 
-	/** The associated engine label. Used by UnrealVersionSelector to open projects with the correct editor */
+	/** The associated engine. Used by UnrealVersionSelector to open projects with the correct editor. Set this value using IDesktopPlatform::SetEngineIdentifierForProject. */
 	FString EngineAssociation;
 
 	/** Category to show under the project browser */
