Code to support 'has feature pack' flag in templatedefs file
Also added code to display an icon on templates that have an equivalent feature pack.
Revised new project text to indicated what the icon is for.

diff --git a/Engine/Source/Editor/EditorStyle/Private/SlateEditorStyle.cpp b/Engine/Source/Editor/EditorStyle/Private/SlateEditorStyle.cpp
--- a/Engine/Source/Editor/EditorStyle/Private/SlateEditorStyle.cpp
+++ b/Engine/Source/Editor/EditorStyle/Private/SlateEditorStyle.cpp
@@ -1607,6 +1607,7 @@ void FSlateEditorStyle::FStyle::SetupGeneralStyles()
 			.SetSelectedTextColor( InvertedForeground )
 			);
 
+		Set( "GameProjectDialog.FeaturePackThumbnail", new IMAGE_BRUSH( "GameProjectDialog/featurepack_avail_32x", FVector2D(32,32) ) );
 		Set( "GameProjectDialog.DefaultGameThumbnail", new IMAGE_BRUSH( "GameProjectDialog/default_game_thumbnail_128x", Icon128x128 ) );
 		Set( "GameProjectDialog.DefaultGameThumbnail.Small", new IMAGE_BRUSH( "GameProjectDialog/default_game_thumbnail", Icon128x128 ) );
 		Set( "GameProjectDialog.BlankProjectThumbnail", new IMAGE_BRUSH( "GameProjectDialog/blank_project_thumbnail", Icon128x128 ) );
diff --git a/Engine/Source/Editor/GameProjectGeneration/Classes/TemplateProjectDefs.h b/Engine/Source/Editor/GameProjectGeneration/Classes/TemplateProjectDefs.h
--- a/Engine/Source/Editor/GameProjectGeneration/Classes/TemplateProjectDefs.h
+++ b/Engine/Source/Editor/GameProjectGeneration/Classes/TemplateProjectDefs.h
@@ -96,6 +96,8 @@ class UTemplateProjectDefs : public UObject
 	UPROPERTY(config)
 	FString AssetTypes;
 
+	UPROPERTY(config)
+	bool bHasFeaturePack;
 
 	/** Fixes up all strings in this definitions object to replace \%TEMPLATENAME\% with the supplied template name and \%PROJECTNAME\% with the supplied project name */
 	void FixupStrings(const FString& TemplateName, const FString& ProjectName);
diff --git a/Engine/Source/Editor/GameProjectGeneration/Private/SNewProjectWizard.cpp b/Engine/Source/Editor/GameProjectGeneration/Private/SNewProjectWizard.cpp
--- a/Engine/Source/Editor/GameProjectGeneration/Private/SNewProjectWizard.cpp
+++ b/Engine/Source/Editor/GameProjectGeneration/Private/SNewProjectWizard.cpp
@@ -44,9 +44,10 @@ struct FTemplateItem
 
 	FString		ClassTypes;
 	FString		AssetTypes;
-	FTemplateItem(FText InName, FText InDescription, bool bInGenerateCode, FName InType, FString InSortKey, FString InProjectFile, TSharedPtr<FSlateBrush> InThumbnail, TSharedPtr<FSlateBrush> InPreviewImage,FString InClassTypes, FString InAssetTypes)
+	bool		bHasFeaturePack;
+	FTemplateItem(FText InName, FText InDescription, bool bInGenerateCode, FName InType, FString InSortKey, FString InProjectFile, TSharedPtr<FSlateBrush> InThumbnail, TSharedPtr<FSlateBrush> InPreviewImage,FString InClassTypes, FString InAssetTypes, bool InHasFeaturePack)
 		: Name(InName), Description(InDescription), bGenerateCode(bInGenerateCode), Type(InType), SortKey(MoveTemp(InSortKey)), ProjectFile(MoveTemp(InProjectFile)), Thumbnail(InThumbnail), PreviewImage(InPreviewImage)
-		, ClassTypes(InClassTypes), AssetTypes(InAssetTypes)
+		, ClassTypes(InClassTypes), AssetTypes(InAssetTypes), bHasFeaturePack(InHasFeaturePack)
 	{}
 };
 
@@ -253,8 +254,20 @@ public:
 					.WidthOverride(ThumbnailSize)
 					.HeightOverride(ThumbnailSize)
 					[
-						SNew(SImage)
-						.Image(this, &STemplateTile::GetThumbnail)
+						SNew(SOverlay)
+						+SOverlay::Slot()
+						[
+							SNew(SImage)
+							.Image(this, &STemplateTile::GetThumbnail)
+						]
+						+ SOverlay::Slot()
+						.HAlign(HAlign_Right)
+						.VAlign(VAlign_Bottom)
+						[
+							SNew(SImage)
+							.Visibility(this, &STemplateTile::GetHasFeaturePackVisibility)
+							.Image(FEditorStyle::GetBrush("GameProjectDialog.FeaturePackThumbnail"))
+						]
 					]
 				]
 
@@ -288,7 +301,18 @@ private:
 		}
 		return FEditorStyle::GetBrush("GameProjectDialog.DefaultGameThumbnail.Small");
 	}
-
+	
+	/** Returns visible if this item has an equivalent feature pack available */
+	EVisibility GetHasFeaturePackVisibility() const
+	{
+		auto ItemPtr = Item.Pin();
+		EVisibility Vis = EVisibility::Hidden;
+		if (ItemPtr.IsValid() )
+		{
+			Vis = ItemPtr->bHasFeaturePack == true ? EVisibility::Visible : EVisibility::Collapsed;
+		}
+		return Vis;
+	}
 };
 
 BEGIN_SLATE_FUNCTION_BUILD_OPTIMIZATION
@@ -555,7 +579,7 @@ void SNewProjectWizard::Construct( const FArguments& InArgs )
 								.Padding(FMargin(0, 0, 0, 15.f))
 								[
 									SNew(STextBlock)
-									.Text(LOCTEXT("ProjectSettingsDescription", "Next, choose some settings for your project. Don't worry, you can choose later or change these at any time in [Project Settings - Target Hardware]:"))
+									.Text(LOCTEXT("ProjectSettingsDescription", "Next, choose some settings for your project. Don't worry, you can choose later or change these at any time in [Project Settings - Target Hardware]\nItems that are marked with an FP icon are availabe as a feature pack and the components of these can be added to a project at any time\nNote also that you can add the starter content after you have created your project."))
 									.ToolTip(IDocumentation::Get()->CreateToolTip(LOCTEXT("HardwareTargetTooltip", "These settings will choose good defaults for a number of other settings in the project such as post-processing flags and touch input emulation using the mouse."), NULL, TEXT("Shared/Editor/NewProjectWizard"), TEXT("TargetHardware")))
 								]
 
@@ -1127,7 +1151,8 @@ void SNewProjectWizard::FindTemplateProjects()
 		MakeShareable( new FSlateBrush( *FEditorStyle::GetBrush("GameProjectDialog.BlankProjectThumbnail") ) ),
 		MakeShareable( new FSlateBrush( *FEditorStyle::GetBrush("GameProjectDialog.BlankProjectPreview") ) ),
 		TEXT(""),		// No class types
-		TEXT("")		// No asset types
+		TEXT(""),		// No asset types,
+		false			// No equivalent feature pack
 		)) );
 
 	Templates.FindOrAdd(FTemplateCategory::CodeCategoryName).Add(MakeShareable(new FTemplateItem(
@@ -1139,7 +1164,8 @@ void SNewProjectWizard::FindTemplateProjects()
 		MakeShareable( new FSlateBrush( *FEditorStyle::GetBrush("GameProjectDialog.BasicCodeThumbnail") ) ),
 		MakeShareable( new FSlateBrush( *FEditorStyle::GetBrush("GameProjectDialog.BlankProjectPreview") ) ),
 		TEXT(""),		// No class types
-		TEXT("")		// No asset types
+		TEXT(""),		// No asset types
+		false			// No equivalent feature pack
 		)) );
 
 	// Now discover and all data driven templates
@@ -1198,6 +1224,7 @@ void SNewProjectWizard::FindTemplateProjects()
 					FText TemplateDescription = TemplateDefs->GetLocalizedDescription();
 					FString ClassTypes = TemplateDefs->ClassTypes;
 					FString AssetTypes = TemplateDefs->AssetTypes;
+					bool bHasFeaturePack = TemplateDefs->bHasFeaturePack;
 
 					// If no template name was specified for the current culture, just use the project name
 					if ( TemplateName.IsEmpty() )
@@ -1252,7 +1279,8 @@ void SNewProjectWizard::FindTemplateProjects()
 						ThumbnailBrush,
 						PreviewBrush,
 						ClassTypes,
-						AssetTypes
+						AssetTypes,
+						bHasFeaturePack
 					)));
 				}
 			}
