Missed marking new files for add

diff --git a/Engine/Source/Editor/GameProjectGeneration/Private/SGetSuggestedIDEWidget.cpp b/Engine/Source/Editor/GameProjectGeneration/Private/SGetSuggestedIDEWidget.cpp
new file mode 100644
--- /dev/null
+++ b/Engine/Source/Editor/GameProjectGeneration/Private/SGetSuggestedIDEWidget.cpp
@@ -0,0 +1,85 @@
+// Copyright 1998-2015 Epic Games, Inc. All Rights Reserved.
+
+
+#include "GameProjectGenerationPrivatePCH.h"
+#include "SourceCodeNavigation.h"
+#include "SNotificationList.h"
+#include "NotificationManager.h"
+#include "SHyperlink.h"
+
+#define LOCTEXT_NAMESPACE "GameProjectGeneration"
+
+void SGetSuggestedIDEWidget::Construct(const FArguments& InArgs)
+{
+	Visibility = TAttribute<EVisibility>(this, &SGetSuggestedIDEWidget::GetVisibility);
+
+	ChildSlot
+	[
+		CreateGetSuggestedIDEWidget()
+	];
+}
+
+TSharedRef<SWidget> SGetSuggestedIDEWidget::CreateGetSuggestedIDEWidget() const
+{
+	if (FSourceCodeNavigation::GetCanDirectlyInstallSourceCodeIDE())
+	{
+		// If the installer for this platform's IDE can be downloaded and launched directly, show a button
+		return SNew(SButton)
+			.Text(FText::Format(LOCTEXT("IDEInstallButtonText", "Install {0}"), FSourceCodeNavigation::GetSuggestedSourceCodeIDE()))
+			.OnClicked(this, &SGetSuggestedIDEWidget::OnInstallIDEClicked);
+	}
+	else
+	{
+		// If the user must open a web page, show a link
+		return SNew(SHyperlink)
+			.Text(FText::Format(LOCTEXT("IDEDownloadLinkText", "Download {0}"), FSourceCodeNavigation::GetSuggestedSourceCodeIDE()))
+			.OnNavigate(this, &SGetSuggestedIDEWidget::OnDownloadIDEClicked, FSourceCodeNavigation::GetSuggestedSourceCodeIDEDownloadURL());
+	}
+}
+
+EVisibility SGetSuggestedIDEWidget::GetVisibility() const
+{
+	/*return FSourceCodeNavigation::IsCompilerAvailable() ? EVisibility::Collapsed : EVisibility::Visible;*/
+	return EVisibility::Visible;
+}
+
+void SGetSuggestedIDEWidget::OnDownloadIDEClicked(FString URL)
+{
+	FPlatformProcess::LaunchURL(*URL, NULL, NULL);
+}
+
+FReply SGetSuggestedIDEWidget::OnInstallIDEClicked()
+{
+	FNotificationInfo Info(FText::Format(LOCTEXT("DownloadingIDEInstaller", "Downloading {0} Installer..."), FSourceCodeNavigation::GetSuggestedSourceCodeIDE(true)));
+	Info.bUseLargeFont = false;
+	Info.bFireAndForget = false;
+	Info.bUseSuccessFailIcons = true;
+	Info.bUseThrobber = true;
+	IDEDownloadNotification = FSlateNotificationManager::Get().AddNotification(Info);
+	IDEDownloadNotification->SetCompletionState(SNotificationItem::ECompletionState::CS_Pending);
+
+	FSourceCodeNavigation::DownloadAndInstallSuggestedIDE(FOnIDEInstallerDownloadComplete::CreateSP(this, &SGetSuggestedIDEWidget::OnIDEInstallerDownloadComplete));
+
+	FEngineAnalytics::GetProvider().RecordEvent(TEXT("Editor.Usage.InstalledIDE"));
+
+	return FReply::Handled();
+}
+
+void SGetSuggestedIDEWidget::OnIDEInstallerDownloadComplete(bool bWasSuccessful)
+{
+	if (IDEDownloadNotification.IsValid())
+	{
+		if (bWasSuccessful)
+		{
+			IDEDownloadNotification->SetCompletionState(SNotificationItem::ECompletionState::CS_Success);
+			IDEDownloadNotification->SetText(LOCTEXT("IDEDownloadSuccess", "Starting installation..."));
+		}
+		else
+		{
+			IDEDownloadNotification->SetCompletionState(SNotificationItem::ECompletionState::CS_Fail);
+			IDEDownloadNotification->SetText(LOCTEXT("IDEDownloadFailure", "Failed to download installer. Please check your internet connection."));
+		}
+
+		IDEDownloadNotification->ExpireAndFadeout();
+	}
+}
\ No newline at end of file
diff --git a/Engine/Source/Editor/GameProjectGeneration/Private/SGetSuggestedIDEWidget.h b/Engine/Source/Editor/GameProjectGeneration/Private/SGetSuggestedIDEWidget.h
new file mode 100644
--- /dev/null
+++ b/Engine/Source/Editor/GameProjectGeneration/Private/SGetSuggestedIDEWidget.h
@@ -0,0 +1,42 @@
+// Copyright 1998-2015 Epic Games, Inc. All Rights Reserved.
+
+#pragma once
+
+#include "SCompoundWidget.h"
+
+/**
+ * Either a button to directly install or a hyperlink to a website to download the suggested IDE for the platform.
+ * Only visible when no compiler is available.
+ */
+class SGetSuggestedIDEWidget : public SCompoundWidget
+{
+public:
+	SLATE_BEGIN_ARGS(SGetSuggestedIDEWidget)
+	{}
+	SLATE_END_ARGS()
+
+	/** Constructs this widget with InArgs */
+	void Construct(const FArguments& InArgs);
+	
+private:
+	/** Creates the appropriate widget to display for the platform */
+	TSharedRef<SWidget> CreateGetSuggestedIDEWidget() const;
+
+	/** Gets the visibility of the global error label IDE Link */
+	EVisibility GetVisibility() const;
+
+	/** Handler for when the error label IDE hyperlink was clicked */
+	void OnDownloadIDEClicked(FString URL);
+
+	/** Handler for when the install IDE button was clicked */
+	FReply OnInstallIDEClicked();
+
+	/** Handler for when the suggested IDE installer has finished downloading */
+	void OnIDEInstallerDownloadComplete(bool bWasSuccessful);
+
+private:
+
+	/** Handle to the notification displayed when downloading an IDE installer */
+	TSharedPtr<class SNotificationItem> IDEDownloadNotification;
+
+};
\ No newline at end of file
