Only auto-compile newly created projects if we're running a Rocket editor. Other editors may require rebuilding binaries, which will fail with the current editor open.

diff --git a/Engine/Source/Editor/GameProjectGeneration/Private/SNewProjectWizard.cpp b/Engine/Source/Editor/GameProjectGeneration/Private/SNewProjectWizard.cpp
--- a/Engine/Source/Editor/GameProjectGeneration/Private/SNewProjectWizard.cpp
+++ b/Engine/Source/Editor/GameProjectGeneration/Private/SNewProjectWizard.cpp
@@ -1491,10 +1491,25 @@ void SNewProjectWizard::CreateAndOpenProject( )
 			// Prevent periodic validity checks. This is to prevent a brief error message about the project already existing while you are exiting.
 			bPreventPeriodicValidityChecksUntilNextChange = true;
 
-			// If it's a code project, compile the binaries first.
-			if(!GetSelectedTemplateItem()->bGenerateCode || GameProjectUtils::BuildCodeProject(ProjectFile))
+			// Rocket already has the engine compiled, so we can try to build and open a new project immediately. Non-Rocket might require building
+			// the engine (especially the case when binaries came from P4), so open the IDE instead.
+			if(FRocketSupport::IsRocket())
 			{
-				OpenProject( ProjectFile );
+				if(!GetSelectedTemplateItem()->bGenerateCode || GameProjectUtils::BuildCodeProject(ProjectFile))
+				{
+					OpenProject(ProjectFile);
+				}
+			}
+			else
+			{
+				if(GetSelectedTemplateItem()->bGenerateCode)
+				{
+					OpenCodeIDE(ProjectFile);
+				}
+				else
+				{
+					OpenProject(ProjectFile);
+				}
 			}
 		}
 	}
