[INTEGRATE] Change 2417468 by [REDACTED_PERSON][REDACTED_HANDLE] on 2015/01/23 18:23:34

	New C++ code templates for Actor and ActorComponents
	- When using "Add C++ Code" in the editor, you'll now get nice starter classes for common class types
	- Classes that inherit from Actor will get a constructor, BeginPlay, and Tick (and will tick by default)
	- Classes inheriting from ActorComponent will get a constructor, InitializeComponent and TickComponent (and will tick by default)

	More improvements coming soon...

diff --git a/Engine/Source/Editor/GameProjectGeneration/Private/GameProjectUtils.cpp b/Engine/Source/Editor/GameProjectGeneration/Private/GameProjectUtils.cpp
--- a/Engine/Source/Editor/GameProjectGeneration/Private/GameProjectUtils.cpp
+++ b/Engine/Source/Editor/GameProjectGeneration/Private/GameProjectUtils.cpp
@@ -246,8 +246,22 @@ FString GameProjectUtils::FNewClassInfo::GetHeaderTemplateFilename() const
 {
 	switch(ClassType)
 	{
-	case EClassType::UObject:
-		return TEXT("UObjectClass.h.template");
+		case EClassType::UObject:
+		{
+			if( BaseClass != nullptr && BaseClass->IsChildOf( UActorComponent::StaticClass() ) )
+			{
+				return TEXT( "ActorComponentClass.h.template" );
+			}
+			else if( BaseClass != nullptr && BaseClass->IsChildOf( AActor::StaticClass() ) )
+			{
+				return TEXT( "ActorClass.h.template" );
+			}
+			else
+			{
+				// Some other non-actor, non-component UObject class
+				return TEXT( "UObjectClass.h.template" );
+			}
+		}
 
 	case EClassType::EmptyCpp:
 		return TEXT("EmptyClass.h.template");
@@ -268,9 +282,21 @@ FString GameProjectUtils::FNewClassInfo::GetSourceTemplateFilename() const
 {
 	switch(ClassType)
 	{
-	case EClassType::UObject:
-		return TEXT("UObjectClass.cpp.template");
-
+		case EClassType::UObject:
+			if( BaseClass != nullptr && BaseClass->IsChildOf( UActorComponent::StaticClass() ) )
+			{
+				return TEXT( "ActorComponentClass.cpp.template" );
+			}
+			else if( BaseClass != nullptr && BaseClass->IsChildOf( AActor::StaticClass() ) )
+			{
+				return TEXT( "ActorClass.cpp.template" );
+			}
+			else
+			{
+				// Some other non-actor, non-component UObject class
+				return TEXT( "UObjectClass.cpp.template" );
+			}
+	
 	case EClassType::EmptyCpp:
 		return TEXT("EmptyClass.cpp.template");
 
