[INTEGRATE] Change 2417804 by [REDACTED_PERSON][REDACTED_HANDLE] on 2015/01/24 11:19:31

	Allow C++ class template source files to have cursor location markup (not just header files)
	- This feature is not currently used but will be in the future

diff --git a/Engine/Source/Editor/GameProjectGeneration/Private/GameProjectUtils.cpp b/Engine/Source/Editor/GameProjectGeneration/Private/GameProjectUtils.cpp
--- a/Engine/Source/Editor/GameProjectGeneration/Private/GameProjectUtils.cpp
+++ b/Engine/Source/Editor/GameProjectGeneration/Private/GameProjectUtils.cpp
@@ -1889,8 +1889,9 @@ bool GameProjectUtils::GenerateGameFrameworkSourceCode(const FString& NewProject
 		
 		TArray<FString> PropertyOverrides;
 		TArray<FString> AdditionalIncludes;
+		FString UnusedSyncLocation;
 
-		if ( GenerateClassCPPFile(NewCPPFilename, NewClassName, FNewClassInfo(BaseClass), AdditionalIncludes, PropertyOverrides, TEXT(""), NewModuleInfo, OutFailReason) )
+		if ( GenerateClassCPPFile(NewCPPFilename, NewClassName, FNewClassInfo(BaseClass), AdditionalIncludes, PropertyOverrides, TEXT(""), UnusedSyncLocation, NewModuleInfo, OutFailReason) )
 		{
 			OutCreatedFiles.Add(NewCPPFilename);
 		}
@@ -2522,34 +2523,12 @@ bool GameProjectUtils::GenerateClassHeaderFile(const FString& NewHeaderFileName,
 	FinalOutput = FinalOutput.Replace(TEXT("%CLASS_FUNCTION_DECLARATIONS%"), *ClassFunctionDeclarations, ESearchCase::CaseSensitive);
 	FinalOutput = FinalOutput.Replace(TEXT("%BASE_CLASS_INCLUDE_DIRECTIVE%"), *BaseClassIncludeDirective, ESearchCase::CaseSensitive);
 
-	// Determine the cursor focus location if this file will by synced after creation
-	TArray<FString> Lines;
-	FinalOutput.ParseIntoArray(&Lines, TEXT("\n"), false);
-	for ( int32 LineIdx = 0; LineIdx < Lines.Num(); ++LineIdx )
-	{
-		const FString& Line = Lines[LineIdx];
-		int32 CharLoc = Line.Find( TEXT("%CURSORFOCUSLOCATION%") );
-		if ( CharLoc != INDEX_NONE )
-		{
-			// Found the sync marker
-			OutSyncLocation = FString::Printf( TEXT("%d:%d"), LineIdx + 1, CharLoc + 1 );
-			break;
-		}
-	}
-
-	// If we did not find the sync location, just sync to the top of the file
-	if ( OutSyncLocation.IsEmpty() )
-	{
-		OutSyncLocation = TEXT("1:1");
-	}
-
-	// Now remove the cursor focus marker
-	FinalOutput = FinalOutput.Replace(TEXT("%CURSORFOCUSLOCATION%"), TEXT(""), ESearchCase::CaseSensitive);
+	HarvestCursorSyncLocation( FinalOutput, OutSyncLocation );
 
 	return WriteOutputFile(NewHeaderFileName, FinalOutput, OutFailReason);
 }
 
-bool GameProjectUtils::GenerateClassCPPFile(const FString& NewCPPFileName, const FString UnPrefixedClassName, const FNewClassInfo ParentClassInfo, const TArray<FString>& AdditionalIncludes, const TArray<FString>& PropertyOverrides, const FString& AdditionalMemberDefinitions, const FModuleContextInfo& ModuleInfo, FText& OutFailReason)
+bool GameProjectUtils::GenerateClassCPPFile(const FString& NewCPPFileName, const FString UnPrefixedClassName, const FNewClassInfo ParentClassInfo, const TArray<FString>& AdditionalIncludes, const TArray<FString>& PropertyOverrides, const FString& AdditionalMemberDefinitions, FString& OutSyncLocation, const FModuleContextInfo& ModuleInfo, FText& OutFailReason)
 {
 	FString Template;
 	if ( !ReadTemplateFile(ParentClassInfo.GetSourceTemplateFilename(), Template, OutFailReason) )
@@ -2613,6 +2592,8 @@ bool GameProjectUtils::GenerateClassCPPFile(const FString& NewCPPFileName, const
 	FinalOutput = FinalOutput.Replace(TEXT("%ADDITIONAL_MEMBER_DEFINITIONS%"), *AdditionalMemberDefinitions, ESearchCase::CaseSensitive);
 	FinalOutput = FinalOutput.Replace(TEXT("%ADDITIONAL_INCLUDE_DIRECTIVES%"), *AdditionalIncludesStr, ESearchCase::CaseSensitive);
 
+	HarvestCursorSyncLocation( FinalOutput, OutSyncLocation );
+
 	return WriteOutputFile(NewCPPFileName, FinalOutput, OutFailReason);
 }
 
@@ -3006,10 +2987,10 @@ bool GameProjectUtils::AddCodeToProject_Internal(const FString& NewClassName, co
 	SlowTask.EnterProgressFrame();
 
 	// Class Header File
-	FString SyncLocation;
 	const FString NewHeaderFilename = NewHeaderPath / ParentClassInfo.GetHeaderFilename(NewClassName);
 	{
-		if ( GenerateClassHeaderFile(NewHeaderFilename, CleanClassName, ParentClassInfo, TArray<FString>(), TEXT(""), TEXT(""), SyncLocation, ModuleInfo, false, OutFailReason) )
+		FString UnusedSyncLocation;
+		if ( GenerateClassHeaderFile(NewHeaderFilename, CleanClassName, ParentClassInfo, TArray<FString>(), TEXT(""), TEXT(""), UnusedSyncLocation, ModuleInfo, false, OutFailReason) )
 		{
 			CreatedFiles.Add(NewHeaderFilename);
 		}
@@ -3025,7 +3006,8 @@ bool GameProjectUtils::AddCodeToProject_Internal(const FString& NewClassName, co
 	// Class CPP file
 	const FString NewCppFilename = NewCppPath / ParentClassInfo.GetSourceFilename(NewClassName);
 	{
-		if ( GenerateClassCPPFile(NewCppFilename, CleanClassName, ParentClassInfo, TArray<FString>(), TArray<FString>(), TEXT(""), ModuleInfo, OutFailReason) )
+		FString UnusedSyncLocation;
+		if ( GenerateClassCPPFile(NewCppFilename, CleanClassName, ParentClassInfo, TArray<FString>(), TArray<FString>(), TEXT(""), UnusedSyncLocation, ModuleInfo, OutFailReason) )
 		{
 			CreatedFiles.Add(NewCppFilename);
 		}
@@ -3146,4 +3128,34 @@ bool GameProjectUtils::FindSourceFileInProject(const FString& InFilename, const
 	return false;
 }
 
+
+void GameProjectUtils::HarvestCursorSyncLocation( FString& FinalOutput, FString& OutSyncLocation )
+{
+	OutSyncLocation.Empty();
+
+	// Determine the cursor focus location if this file will by synced after creation
+	TArray<FString> Lines;
+	FinalOutput.ParseIntoArray( &Lines, TEXT( "\n" ), false );
+	for( int32 LineIdx = 0; LineIdx < Lines.Num(); ++LineIdx )
+	{
+		const FString& Line = Lines[ LineIdx ];
+		int32 CharLoc = Line.Find( TEXT( "%CURSORFOCUSLOCATION%" ) );
+		if( CharLoc != INDEX_NONE )
+		{
+			// Found the sync marker
+			OutSyncLocation = FString::Printf( TEXT( "%d:%d" ), LineIdx + 1, CharLoc + 1 );
+			break;
+		}
+	}
+
+	// If we did not find the sync location, just sync to the top of the file
+	if( OutSyncLocation.IsEmpty() )
+	{
+		OutSyncLocation = TEXT( "1:1" );
+	}
+
+	// Now remove the cursor focus marker
+	FinalOutput = FinalOutput.Replace(TEXT("%CURSORFOCUSLOCATION%"), TEXT(""), ESearchCase::CaseSensitive);
+}
+
 #undef LOCTEXT_NAMESPACE
diff --git a/Engine/Source/Editor/GameProjectGeneration/Private/GameProjectUtils.h b/Engine/Source/Editor/GameProjectGeneration/Private/GameProjectUtils.h
--- a/Engine/Source/Editor/GameProjectGeneration/Private/GameProjectUtils.h
+++ b/Engine/Source/Editor/GameProjectGeneration/Private/GameProjectUtils.h
@@ -364,8 +364,11 @@ private:
 	/** Generates a header file for a UObject class. OutSyncLocation is a string representing the preferred cursor sync location for this file after creation. */
 	static bool GenerateClassHeaderFile(const FString& NewHeaderFileName, const FString UnPrefixedClassName, const FNewClassInfo ParentClassInfo, const TArray<FString>& ClassSpecifierList, const FString& ClassProperties, const FString& ClassFunctionDeclarations, FString& OutSyncLocation, const FModuleContextInfo& ModuleInfo, bool bDeclareConstructor, FText& OutFailReason);
 
+	/** Finds the cursor sync location in the source file and reports it back as a string */
+	static void HarvestCursorSyncLocation( FString& FinalOutput, FString& OutSyncLocation );
+
 	/** Generates a cpp file for a UObject class */
-	static bool GenerateClassCPPFile(const FString& NewCPPFileName, const FString UnPrefixedClassName, const FNewClassInfo ParentClassInfo, const TArray<FString>& AdditionalIncludes, const TArray<FString>& PropertyOverrides, const FString& AdditionalMemberDefinitions, const FModuleContextInfo& ModuleInfo, FText& OutFailReason);
+	static bool GenerateClassCPPFile(const FString& NewCPPFileName, const FString UnPrefixedClassName, const FNewClassInfo ParentClassInfo, const TArray<FString>& AdditionalIncludes, const TArray<FString>& PropertyOverrides, const FString& AdditionalMemberDefinitions, FString& OutSyncLocation, const FModuleContextInfo& ModuleInfo, FText& OutFailReason);
 
 	/** Generates a Build.cs file for a game module */
 	static bool GenerateGameModuleBuildFile(const FString& NewBuildFileName, const FString& ModuleName, const TArray<FString>& PublicDependencyModuleNames, const TArray<FString>& PrivateDependencyModuleNames, FText& OutFailReason);
