[INTEGRATE] Change 2417820 by [REDACTED_PERSON][REDACTED_HANDLE] on 2015/01/24 11:57:38

	After adding a new class, we now always open the IDE to edit source if the hot reload completed successfully
	- Also, we display a quick pop-up notification indicating success instead of a modal prompt

diff --git a/Engine/Source/Editor/GameProjectGeneration/Private/SNewClassDialog.cpp b/Engine/Source/Editor/GameProjectGeneration/Private/SNewClassDialog.cpp
--- a/Engine/Source/Editor/GameProjectGeneration/Private/SNewClassDialog.cpp
+++ b/Engine/Source/Editor/GameProjectGeneration/Private/SNewClassDialog.cpp
@@ -23,6 +23,8 @@
 #include "GameFramework/HUD.h"
 #include "GameFramework/PlayerState.h"
 #include "Kismet/BlueprintFunctionLibrary.h"
+#include "SNotificationList.h"
+#include "NotificationManager.h"
 
 
 #define LOCTEXT_NAMESPACE "GameProjectGeneration"
@@ -997,25 +999,46 @@ void SNewClassDialog::FinishClicked()
 		bPreventPeriodicValidityChecksUntilNextChange = true;
 
 		// Display a nag if we didn't automatically hot-reload for the newly added class
-		const bool bShowHotReloadNag = !GEditor->AccessEditorUserSettings().bAutomaticallyHotReloadNewClasses;
+		const bool bWasHotReloaded = GEditor->AccessEditorUserSettings().bAutomaticallyHotReloadNewClasses;
+		if( bWasHotReloaded )
+		{
+			FNotificationInfo Notification( FText::Format( LOCTEXT("AddedClassSuccessNotification", "Added new class {0}"), FText::FromString(NewClassName) ) );
+			FSlateNotificationManager::Get().AddNotification( Notification );
+		}
 
 		if ( HeaderFilePath.IsEmpty() || CppFilePath.IsEmpty() || !FSlateApplication::Get().SupportsSourceAccess() )
 		{
-			// Code successfully added, notify the user. We are either running on a platform that does not support source access or a file was not given so don't ask about editing the file
-			const FText Message = FText::Format( (bShowHotReloadNag) 
-				? LOCTEXT("AddCodeSuccessWithHotReload", "Successfully added class {0}, however you must recompile {1} before it will appear in the Content Browser.")
-				: LOCTEXT("AddCodeSuccess", "Successfully added class {0}.")
-				, FText::FromString(NewClassName), FText::FromString(SelectedModuleInfo->ModuleName) );
-			FMessageDialog::Open(EAppMsgType::Ok, Message);
+			if( !bWasHotReloaded )
+			{
+				// Code successfully added, notify the user. We are either running on a platform that does not support source access or a file was not given so don't ask about editing the file
+				const FText Message = FText::Format( 
+					LOCTEXT("AddCodeSuccessWithHotReload", "Successfully added class {0}, however you must recompile {1} before it will appear in the Content Browser.")
+					, FText::FromString(NewClassName), FText::FromString(SelectedModuleInfo->ModuleName) );
+				FMessageDialog::Open(EAppMsgType::Ok, Message);
+			}
+			else
+			{
+				// Code was added and hot reloaded into the editor, but the user doesn't have a code IDE installed so we can't open the file to edit it now
+			}
 		}
 		else
 		{
-			// Code successfully added, notify the user and ask about opening the IDE now
-			const FText Message = FText::Format( (bShowHotReloadNag) 
-				? LOCTEXT("AddCodeSuccessWithHotReloadAndSync", "Successfully added class {0}, however you must recompile {1} before it will appear in the Content Browser.\n\nWould you like to edit the code now?")
-				: LOCTEXT("AddCodeSuccessWithSync", "Successfully added class {0}. Would you like to edit the code now?")
-				, FText::FromString(NewClassName), FText::FromString(SelectedModuleInfo->ModuleName) );
-			if ( FMessageDialog::Open(EAppMsgType::YesNo, Message) == EAppReturnType::Yes )
+			bool bEditSourceFilesNow = false;
+			if( bWasHotReloaded )
+			{
+				// Code was hot reloaded, so always edit the new classes now
+				bEditSourceFilesNow = true;
+			}
+			else
+			{
+				// Code successfully added, notify the user and ask about opening the IDE now
+				const FText Message = FText::Format( 
+					LOCTEXT("AddCodeSuccessWithHotReloadAndSync", "Successfully added class {0}, however you must recompile {1} before it will appear in the Content Browser.\n\nWould you like to edit the code now?")
+					, FText::FromString(NewClassName), FText::FromString(SelectedModuleInfo->ModuleName) );
+				bEditSourceFilesNow = ( FMessageDialog::Open( EAppMsgType::YesNo, Message ) == EAppReturnType::Yes );
+			}
+
+			if( bEditSourceFilesNow )
 			{
 				TArray<FString> SourceFiles;
 				SourceFiles.Add(IFileManager::Get().ConvertToAbsolutePathForExternalAppForRead(*HeaderFilePath));
