[INTEGRATE] Change 2418552 by [REDACTED_PERSON][REDACTED_HANDLE] on 2015/01/26 08:58:56

	UE-7891 Remove 'feature pack' icons, and change text to explain that any template features can be added later.

diff --git a/Engine/Source/Editor/EditorStyle/Private/SlateEditorStyle.cpp b/Engine/Source/Editor/EditorStyle/Private/SlateEditorStyle.cpp
--- a/Engine/Source/Editor/EditorStyle/Private/SlateEditorStyle.cpp
+++ b/Engine/Source/Editor/EditorStyle/Private/SlateEditorStyle.cpp
@@ -1671,7 +1671,6 @@ void FSlateEditorStyle::FStyle::SetupGeneralStyles()
 			.SetSelectedTextColor( InvertedForeground )
 			);
 
-		Set( "GameProjectDialog.FeaturePackThumbnail", new IMAGE_BRUSH( "GameProjectDialog/featurepack_avail_32x", FVector2D(32,32) ) );
 		Set( "GameProjectDialog.DefaultGameThumbnail", new IMAGE_BRUSH( "GameProjectDialog/default_game_thumbnail_128x", Icon128x128 ) );
 		Set( "GameProjectDialog.DefaultGameThumbnail.Small", new IMAGE_BRUSH( "GameProjectDialog/default_game_thumbnail", Icon128x128 ) );
 		Set( "GameProjectDialog.BlankProjectThumbnail", new IMAGE_BRUSH( "GameProjectDialog/blank_project_thumbnail", Icon128x128 ) );
diff --git a/Engine/Source/Editor/GameProjectGeneration/Private/SNewProjectWizard.cpp b/Engine/Source/Editor/GameProjectGeneration/Private/SNewProjectWizard.cpp
--- a/Engine/Source/Editor/GameProjectGeneration/Private/SNewProjectWizard.cpp
+++ b/Engine/Source/Editor/GameProjectGeneration/Private/SNewProjectWizard.cpp
@@ -254,20 +254,8 @@ public:
 					.WidthOverride(ThumbnailSize)
 					.HeightOverride(ThumbnailSize)
 					[
-						SNew(SOverlay)
-						+SOverlay::Slot()
-						[
-							SNew(SImage)
-							.Image(this, &STemplateTile::GetThumbnail)
-						]
-						+ SOverlay::Slot()
-						.HAlign(HAlign_Right)
-						.VAlign(VAlign_Bottom)
-						[
-							SNew(SImage)
-							.Visibility(this, &STemplateTile::GetHasFeaturePackVisibility)
-							.Image(FEditorStyle::GetBrush("GameProjectDialog.FeaturePackThumbnail"))
-						]
+						SNew(SImage)
+						.Image(this, &STemplateTile::GetThumbnail)
 					]
 				]
 
@@ -302,17 +290,6 @@ private:
 		return FEditorStyle::GetBrush("GameProjectDialog.DefaultGameThumbnail.Small");
 	}
 	
-	/** Returns visible if this item has an equivalent feature pack available */
-	EVisibility GetHasFeaturePackVisibility() const
-	{
-		auto ItemPtr = Item.Pin();
-		EVisibility Vis = EVisibility::Hidden;
-		if (ItemPtr.IsValid() )
-		{
-			Vis = ItemPtr->bHasFeaturePack == true ? EVisibility::Visible : EVisibility::Collapsed;
-		}
-		return Vis;
-	}
 };
 
 BEGIN_SLATE_FUNCTION_BUILD_OPTIMIZATION
@@ -395,11 +372,11 @@ void SNewProjectWizard::Construct( const FArguments& InArgs )
 						SNew(SVerticalBox)
 
 						+ SVerticalBox::Slot()
-						//.Padding(FMargin(0, 0, 0, 15))
+						.Padding(FMargin(0, 0, 0, 15))
 						.AutoHeight()
 						[
 							SNew(STextBlock)
-							.Text(LOCTEXT("ProjectTemplateDescription", "First, choose a template to use as a starting point for your new project:"))
+							.Text(LOCTEXT("ProjectTemplateDescription", "First, choose a template to use as a starting point for your new project.\nAny of these features can also be added later using the [Add New Feature Pack] option."))
 							.ToolTip(IDocumentation::Get()->CreateToolTip(LOCTEXT("TemplateChoiceTooltip", "A template consists of a little bit of player control logic (either as a Blueprint or in C++), input bindings, and appropriate prototyping assets."), NULL, TEXT("Shared/Editor/NewProjectWizard"), TEXT("TemplateChoice")))
 						]
 
@@ -579,7 +556,7 @@ void SNewProjectWizard::Construct( const FArguments& InArgs )
 								.Padding(FMargin(0, 0, 0, 15.f))
 								[
 									SNew(STextBlock)
-									.Text(LOCTEXT("ProjectSettingsDescription", "Next, choose some settings for your project. Don't worry, you can choose later or change these at any time in [Project Settings - Target Hardware]\nItems that are marked with an FP icon are availabe as a feature pack and the components of these can be added to a project at any time\nNote also that you can add the starter content after you have created your project."))
+									.Text(LOCTEXT("ProjectSettingsDescription", "Next, choose some settings for your project. Don't worry, you can choose later or change these at any time in [Project Settings - Target Hardware]\nNote also that you can add the starter content after you have created your project."))
 									.ToolTip(IDocumentation::Get()->CreateToolTip(LOCTEXT("HardwareTargetTooltip", "These settings will choose good defaults for a number of other settings in the project such as post-processing flags and touch input emulation using the mouse."), NULL, TEXT("Shared/Editor/NewProjectWizard"), TEXT("TargetHardware")))
 								]
 
