[INTEGRATE] Change 2418609 by bruce.nesbit@bnesbit_Releases on 2015/01/26 10:18:10

	Removed some now redundant code relating to new projects having a feature pack.

diff --git a/Engine/Source/Editor/GameProjectGeneration/Private/SNewProjectWizard.cpp b/Engine/Source/Editor/GameProjectGeneration/Private/SNewProjectWizard.cpp
--- a/Engine/Source/Editor/GameProjectGeneration/Private/SNewProjectWizard.cpp
+++ b/Engine/Source/Editor/GameProjectGeneration/Private/SNewProjectWizard.cpp
@@ -44,10 +44,9 @@ struct FTemplateItem
 
 	FString		ClassTypes;
 	FString		AssetTypes;
-	bool		bHasFeaturePack;
-	FTemplateItem(FText InName, FText InDescription, bool bInGenerateCode, FName InType, FString InSortKey, FString InProjectFile, TSharedPtr<FSlateBrush> InThumbnail, TSharedPtr<FSlateBrush> InPreviewImage,FString InClassTypes, FString InAssetTypes, bool InHasFeaturePack)
+	FTemplateItem(FText InName, FText InDescription, bool bInGenerateCode, FName InType, FString InSortKey, FString InProjectFile, TSharedPtr<FSlateBrush> InThumbnail, TSharedPtr<FSlateBrush> InPreviewImage,FString InClassTypes, FString InAssetTypes)
 		: Name(InName), Description(InDescription), bGenerateCode(bInGenerateCode), Type(InType), SortKey(MoveTemp(InSortKey)), ProjectFile(MoveTemp(InProjectFile)), Thumbnail(InThumbnail), PreviewImage(InPreviewImage)
-		, ClassTypes(InClassTypes), AssetTypes(InAssetTypes), bHasFeaturePack(InHasFeaturePack)
+		, ClassTypes(InClassTypes), AssetTypes(InAssetTypes)
 	{}
 };
 
@@ -1131,8 +1130,7 @@ void SNewProjectWizard::FindTemplateProjects()
 		MakeShareable( new FSlateBrush( *FEditorStyle::GetBrush("GameProjectDialog.BlankProjectThumbnail") ) ),
 		MakeShareable( new FSlateBrush( *FEditorStyle::GetBrush("GameProjectDialog.BlankProjectPreview") ) ),
 		TEXT(""),		// No class types
-		TEXT(""),		// No asset types,
-		false			// No equivalent feature pack
+		TEXT("")		// No asset types
 		)) );
 
 	Templates.FindOrAdd(FTemplateCategory::CodeCategoryName).Add(MakeShareable(new FTemplateItem(
@@ -1144,8 +1142,7 @@ void SNewProjectWizard::FindTemplateProjects()
 		MakeShareable( new FSlateBrush( *FEditorStyle::GetBrush("GameProjectDialog.BasicCodeThumbnail") ) ),
 		MakeShareable( new FSlateBrush( *FEditorStyle::GetBrush("GameProjectDialog.BlankProjectPreview") ) ),
 		TEXT(""),		// No class types
-		TEXT(""),		// No asset types
-		false			// No equivalent feature pack
+		TEXT("")		// No asset types
 		)) );
 
 	// Now discover and all data driven templates
@@ -1204,8 +1201,7 @@ void SNewProjectWizard::FindTemplateProjects()
 					FText TemplateDescription = TemplateDefs->GetLocalizedDescription();
 					FString ClassTypes = TemplateDefs->ClassTypes;
 					FString AssetTypes = TemplateDefs->AssetTypes;
-					bool bHasFeaturePack = TemplateDefs->bHasFeaturePack;
-
+					
 					// If no template name was specified for the current culture, just use the project name
 					if ( TemplateName.IsEmpty() )
 					{
@@ -1259,8 +1255,7 @@ void SNewProjectWizard::FindTemplateProjects()
 						ThumbnailBrush,
 						PreviewBrush,
 						ClassTypes,
-						AssetTypes,
-						bHasFeaturePack
+						AssetTypes
 					)));
 				}
 			}
