[INTEGRATE] Change 2418633 by [REDACTED_PERSON][REDACTED_HANDLE] on 2015/01/26 10:53:14

	Fix regression with Class Wizard-generated source files not compiling sometimes
	- Now, only AActor, USceneComponent and UActorComponent get "special" generated source
	- FObjectInitializer constructors are no longer used in code generated for these classes

diff --git a/Engine/Source/Editor/GameProjectGeneration/Private/GameProjectUtils.cpp b/Engine/Source/Editor/GameProjectGeneration/Private/GameProjectUtils.cpp
--- a/Engine/Source/Editor/GameProjectGeneration/Private/GameProjectUtils.cpp
+++ b/Engine/Source/Editor/GameProjectGeneration/Private/GameProjectUtils.cpp
@@ -248,11 +248,11 @@ FString GameProjectUtils::FNewClassInfo::GetHeaderTemplateFilename() const
 	{
 		case EClassType::UObject:
 		{
-			if( BaseClass != nullptr && BaseClass->IsChildOf( UActorComponent::StaticClass() ) )
+			if( BaseClass != nullptr && ( BaseClass == UActorComponent::StaticClass() || BaseClass == USceneComponent::StaticClass() ) )
 			{
 				return TEXT( "ActorComponentClass.h.template" );
 			}
-			else if( BaseClass != nullptr && BaseClass->IsChildOf( AActor::StaticClass() ) )
+			else if( BaseClass != nullptr && BaseClass == AActor::StaticClass() )
 			{
 				return TEXT( "ActorClass.h.template" );
 			}
@@ -283,11 +283,11 @@ FString GameProjectUtils::FNewClassInfo::GetSourceTemplateFilename() const
 	switch(ClassType)
 	{
 		case EClassType::UObject:
-			if( BaseClass != nullptr && BaseClass->IsChildOf( UActorComponent::StaticClass() ) )
+			if( BaseClass != nullptr && ( BaseClass == UActorComponent::StaticClass() || BaseClass == USceneComponent::StaticClass() ) )
 			{
 				return TEXT( "ActorComponentClass.cpp.template" );
 			}
-			else if( BaseClass != nullptr && BaseClass->IsChildOf( AActor::StaticClass() ) )
+			else if( BaseClass != nullptr && BaseClass == AActor::StaticClass() )
 			{
 				return TEXT( "ActorClass.cpp.template" );
 			}
