[INTEGRATE] Change 2423672 by [REDACTED_PERSON][REDACTED_HANDLE] on 2015/01/29 08:58:34

	Disable starter content by default.

diff --git a/Engine/Source/Editor/GameProjectGeneration/Private/SNewProjectWizard.cpp b/Engine/Source/Editor/GameProjectGeneration/Private/SNewProjectWizard.cpp
--- a/Engine/Source/Editor/GameProjectGeneration/Private/SNewProjectWizard.cpp
+++ b/Engine/Source/Editor/GameProjectGeneration/Private/SNewProjectWizard.cpp
@@ -299,7 +299,7 @@ void SNewProjectWizard::Construct( const FArguments& InArgs )
 	bLastGlobalValidityCheckSuccessful = true;
 	bLastNameAndLocationValidityCheckSuccessful = true;
 	bPreventPeriodicValidityChecksUntilNextChange = false;
-	bCopyStarterContent = true;
+	bCopyStarterContent = false;
 
 	IHardwareTargetingModule& HardwareTargeting = IHardwareTargetingModule::Get();
 
