[INTEGRATE] Change 2425651 by [REDACTED_PERSON][REDACTED_HANDLE] on 2015/01/30 09:23:54

	New Project Wizard UI text clean-up
	- Switched to rich text and use a bold highlight on some words instead of brackets
	- Turned on automatic text wrapping and removed newlines
	- Re-phrased some of the field descriptions to read more nicely

diff --git a/Engine/Source/Editor/EditorStyle/Private/SlateEditorStyle.cpp b/Engine/Source/Editor/EditorStyle/Private/SlateEditorStyle.cpp
--- a/Engine/Source/Editor/EditorStyle/Private/SlateEditorStyle.cpp
+++ b/Engine/Source/Editor/EditorStyle/Private/SlateEditorStyle.cpp
@@ -184,6 +184,11 @@ void FSlateEditorStyle::FStyle::SetupGeneralStyles()
 	{
 		Set( "RichTextBlock.TextHighlight", FTextBlockStyle(NormalText)
 			.SetColorAndOpacity( FLinearColor( 1.0f, 1.0f, 1.0f ) ) );
+		Set( "RichTextBlock.Bold", FTextBlockStyle(NormalText)
+			.SetFont( TTF_CORE_FONT("Fonts/Roboto-Bold", 10 )) );
+		Set( "RichTextBlock.BoldHighlight", FTextBlockStyle(NormalText)
+			.SetFont( TTF_CORE_FONT("Fonts/Roboto-Bold", 10 ))
+			.SetColorAndOpacity( FLinearColor( 1.0f, 1.0f, 1.0f ) ) );
 
 		Set( "TextBlock.HighlightShape",  new BOX_BRUSH( "Common/TextBlockHighlightShape", FMargin(3.f/8.f) ));
 		Set( "TextBlock.HighlighColor", FLinearColor( 0.02f, 0.3f, 0.0f ) );
diff --git a/Engine/Source/Editor/GameProjectGeneration/Private/SNewProjectWizard.cpp b/Engine/Source/Editor/GameProjectGeneration/Private/SNewProjectWizard.cpp
--- a/Engine/Source/Editor/GameProjectGeneration/Private/SNewProjectWizard.cpp
+++ b/Engine/Source/Editor/GameProjectGeneration/Private/SNewProjectWizard.cpp
@@ -374,8 +374,10 @@ void SNewProjectWizard::Construct( const FArguments& InArgs )
 						.Padding(FMargin(0, 0, 0, 15))
 						.AutoHeight()
 						[
-							SNew(STextBlock)
-							.Text(LOCTEXT("ProjectTemplateDescription", "First, choose a template to use as a starting point for your new project.\nAny of these features can also be added later using the [Add New Feature Pack] option."))
+							SNew(SRichTextBlock)
+							.Text(LOCTEXT("ProjectTemplateDescription", "Choose a <RichTextBlock.BoldHighlight>template</> to use as a starting point for your new project.  Any of these features can be added later by clicking <RichTextBlock.BoldHighlight>Add New Feature Pack</> in <RichTextBlock.BoldHighlight>Content Browser</>."))
+							.AutoWrapText(true)
+ 							.DecoratorStyleSet(&FEditorStyle::Get())
 							.ToolTip(IDocumentation::Get()->CreateToolTip(LOCTEXT("TemplateChoiceTooltip", "A template consists of a little bit of player control logic (either as a Blueprint or in C++), input bindings, and appropriate prototyping assets."), NULL, TEXT("Shared/Editor/NewProjectWizard"), TEXT("TemplateChoice")))
 						]
 
@@ -554,8 +556,10 @@ void SNewProjectWizard::Construct( const FArguments& InArgs )
 								.AutoHeight()
 								.Padding(FMargin(0, 0, 0, 15.f))
 								[
-									SNew(STextBlock)
-									.Text(LOCTEXT("ProjectSettingsDescription", "Next, choose some settings for your project. Don't worry, you can choose later or change these at any time in [Project Settings - Target Hardware]\nNote also that you can add the starter content after you have created your project."))
+									SNew(SRichTextBlock)
+									.Text(LOCTEXT("ProjectSettingsDescription", "Choose some <RichTextBlock.BoldHighlight>settings</> for your project.  Don't worry, you can change these later in the <RichTextBlock.BoldHighlight>Target Hardware</> section of <RichTextBlock.BoldHighlight>Project Settings</>.  You can also add the <RichTextBlock.BoldHighlight>Starter Content</> to your project later using <RichTextBlock.BoldHighlight>Content Browser</>."))
+									.AutoWrapText(true)
+ 									.DecoratorStyleSet(&FEditorStyle::Get())
 									.ToolTip(IDocumentation::Get()->CreateToolTip(LOCTEXT("HardwareTargetTooltip", "These settings will choose good defaults for a number of other settings in the project such as post-processing flags and touch input emulation using the mouse."), NULL, TEXT("Shared/Editor/NewProjectWizard"), TEXT("TargetHardware")))
 								]
 
@@ -624,8 +628,10 @@ void SNewProjectWizard::Construct( const FArguments& InArgs )
 								.AutoHeight()
 								.Padding(FMargin(0, 0, 0, 15.f))
 								[
-									SNew(STextBlock)
-									.Text(LOCTEXT("ProjectPathDescription", "Finally, choose a location for your project to be stored:"))
+									SNew(SRichTextBlock)
+									.Text(LOCTEXT("ProjectPathDescription", "Select a <RichTextBlock.BoldHighlight>location</> for your project to be stored."))
+									.AutoWrapText(true)
+ 									.DecoratorStyleSet(&FEditorStyle::Get())
 									.ToolTip(IDocumentation::Get()->CreateToolTip(LOCTEXT("ProjectPathDescriptionTooltip", "All of your project content and code will be stored here."), NULL, TEXT("Shared/Editor/NewProjectWizard"), TEXT("ProjectPath")))
 								]
 
