[INTEGRATE] Change 2428764 by [REDACTED_PERSON][REDACTED_HANDLE] on 2015/02/02 17:10:37

	Always open Visual Studio or Xcode for newly-created code projects
	- Also, fixed a minor bug with text in the New Project Wizard (UE-8859)

diff --git a/Engine/Source/Editor/GameProjectGeneration/Private/SNewProjectWizard.cpp b/Engine/Source/Editor/GameProjectGeneration/Private/SNewProjectWizard.cpp
--- a/Engine/Source/Editor/GameProjectGeneration/Private/SNewProjectWizard.cpp
+++ b/Engine/Source/Editor/GameProjectGeneration/Private/SNewProjectWizard.cpp
@@ -375,7 +375,7 @@ void SNewProjectWizard::Construct( const FArguments& InArgs )
 						.AutoHeight()
 						[
 							SNew(SRichTextBlock)
-							.Text(LOCTEXT("ProjectTemplateDescription", "Choose a <RichTextBlock.BoldHighlight>template</> to use as a starting point for your new project.  Any of these features can be added later by clicking <RichTextBlock.BoldHighlight>Add New Feature Pack</> in <RichTextBlock.BoldHighlight>Content Browser</>."))
+							.Text(LOCTEXT("ProjectTemplateDescription", "Choose a <RichTextBlock.BoldHighlight>template</> to use as a starting point for your new project.  Any of these features can be added later by clicking <RichTextBlock.BoldHighlight>Add Feature or Content Pack</> in <RichTextBlock.BoldHighlight>Content Browser</>."))
 							.AutoWrapText(true)
  							.DecoratorStyleSet(&FEditorStyle::Get())
 							.ToolTip(IDocumentation::Get()->CreateToolTip(LOCTEXT("TemplateChoiceTooltip", "A template consists of a little bit of player control logic (either as a Blueprint or in C++), input bindings, and appropriate prototyping assets."), NULL, TEXT("Shared/Editor/NewProjectWizard"), TEXT("TemplateChoice")))
@@ -1469,25 +1469,29 @@ void SNewProjectWizard::CreateAndOpenProject( )
 			// Prevent periodic validity checks. This is to prevent a brief error message about the project already existing while you are exiting.
 			bPreventPeriodicValidityChecksUntilNextChange = true;
 
-			// Rocket already has the engine compiled, so we can try to build and open a new project immediately. Non-Rocket might require building
-			// the engine (especially the case when binaries came from P4), so open the IDE instead.
-			if(FRocketSupport::IsRocket())
+			bool bCanOpenProject = false;
+			if( GetSelectedTemplateItem()->bGenerateCode )
 			{
-				if(!GetSelectedTemplateItem()->bGenerateCode || GameProjectUtils::BuildCodeProject(ProjectFile))
-				{
-					OpenProject(ProjectFile);
-				}
+			    // Open Visual Studio or Xcode if the user created a project with C++ files.  Even if the project doesn't compile, we'll still open it in the IDE
+				OpenCodeIDE( ProjectFile );
+
+			    // Rocket already has the engine compiled, so we can try to build and open a new project immediately. Non-Rocket might require building
+			    // the engine (especially the case when binaries came from P4), so we only open the IDE for that.
+			    if( FRocketSupport::IsRocket() && GameProjectUtils::BuildCodeProject(ProjectFile) )
+			    {
+					// Everything compiled OK, so we can go ahead and open the project
+					bCanOpenProject = true;
+			    }
 			}
 			else
 			{
-				if(GetSelectedTemplateItem()->bGenerateCode)
-				{
-					OpenCodeIDE(ProjectFile);
-				}
-				else
-				{
-					OpenProject(ProjectFile);
-				}
+				// We can always open non-code projects, because they don't have any DLLs
+				bCanOpenProject = true;
+			}
+
+			if( bCanOpenProject )
+			{
+				OpenProject(ProjectFile);
 			}
 		}
 	}
