[INTEGRATE] Change 2429834 by [REDACTED_PERSON][REDACTED_HANDLE] on 2015/02/03 08:58:13

	Fixed adding a new class via the editor reaching 100% while it was still compiling the code

	An extra step had been added to the task, but the number of sub-tasks had been left at 6.

diff --git a/Engine/Source/Editor/GameProjectGeneration/Private/GameProjectUtils.cpp b/Engine/Source/Editor/GameProjectGeneration/Private/GameProjectUtils.cpp
--- a/Engine/Source/Editor/GameProjectGeneration/Private/GameProjectUtils.cpp
+++ b/Engine/Source/Editor/GameProjectGeneration/Private/GameProjectUtils.cpp
@@ -2979,7 +2979,7 @@ bool GameProjectUtils::AddCodeToProject_Internal(const FString& NewClassName, co
 		return false;
 	}
 
-	FScopedSlowTask SlowTask( 6, LOCTEXT( "AddingCodeToProject", "Adding code to project..." ) );
+	FScopedSlowTask SlowTask( 7, LOCTEXT( "AddingCodeToProject", "Adding code to project..." ) );
 	SlowTask.MakeDialog();
 
 	SlowTask.EnterProgressFrame();
