[INTEGRATE] Change 2431516 by bruce.nesbit@bnesbit_Releases on 2015/02/04 06:23:58

	Revised starter content so it is inserted into projects using feature packs.
	Removed StarterContent source data from rocket build.

diff --git a/Engine/Source/Editor/GameProjectGeneration/Classes/TemplateProjectDefs.h b/Engine/Source/Editor/GameProjectGeneration/Classes/TemplateProjectDefs.h
--- a/Engine/Source/Editor/GameProjectGeneration/Classes/TemplateProjectDefs.h
+++ b/Engine/Source/Editor/GameProjectGeneration/Classes/TemplateProjectDefs.h
@@ -98,6 +98,10 @@ class UTemplateProjectDefs : public UObject
 
 	UPROPERTY(config)
 	bool bHasFeaturePack;
+	
+	/* Optional list of feature packs to include */
+	UPROPERTY(config)
+	TArray<FString> PacksToInclude;
 
 	/** Fixes up all strings in this definitions object to replace \%TEMPLATENAME\% with the supplied template name and \%PROJECTNAME\% with the supplied project name */
 	void FixupStrings(const FString& TemplateName, const FString& ProjectName);
diff --git a/Engine/Source/Editor/GameProjectGeneration/GameProjectGeneration.Build.cs b/Engine/Source/Editor/GameProjectGeneration/GameProjectGeneration.Build.cs
--- a/Engine/Source/Editor/GameProjectGeneration/GameProjectGeneration.Build.cs
+++ b/Engine/Source/Editor/GameProjectGeneration/GameProjectGeneration.Build.cs
@@ -32,7 +32,7 @@ public class GameProjectGeneration : ModuleRules
  				"TargetPlatform",
 				"UnrealEd",
 				"DesktopPlatform",
-                "HardwareTargeting"
+                "HardwareTargeting",
 			}
 		);
 
diff --git a/Engine/Source/Editor/GameProjectGeneration/Private/GameProjectUtils.cpp b/Engine/Source/Editor/GameProjectGeneration/Private/GameProjectUtils.cpp
--- a/Engine/Source/Editor/GameProjectGeneration/Private/GameProjectUtils.cpp
+++ b/Engine/Source/Editor/GameProjectGeneration/Private/GameProjectUtils.cpp
@@ -25,7 +25,7 @@
 #include "HotReloadInterface.h"
 #include "SVerbChoiceDialog.h"
 #include "SourceCodeNavigation.h"
-
+///#include "AssetToolsModule.h"
 
 #define LOCTEXT_NAMESPACE "GameProjectUtils"
 
@@ -536,58 +536,11 @@ bool GameProjectUtils::OpenCodeIDE(const FString& ProjectFile, FText& OutFailRea
 
 void GameProjectUtils::GetStarterContentFiles(TArray<FString>& OutFilenames)
 {
-	FString const SrcFolder = FPaths::StarterContentDir();
-	FString const ContentFolder = SrcFolder / TEXT("Content");
-
-	// only copying /Content
-	IFileManager::Get().FindFilesRecursive(OutFilenames, *ContentFolder, TEXT("*"), /*Files=*/true, /*Directories=*/false);
-}
-
-bool GameProjectUtils::CopyStarterContent(const FString& DestProjectFolder, FText& OutFailReason)
-{
-	FString const SrcFolder = FPaths::StarterContentDir();
-
-	TArray<FString> FilesToCopy;
-	GetStarterContentFiles(FilesToCopy);
-
-	FScopedSlowTask SlowTask(FilesToCopy.Num(), LOCTEXT("CreatingProjectStatus_CopyingFiles", "Copying Files {SrcFilename}..."));
-	SlowTask.MakeDialog();
-
-	TArray<FString> CreatedFiles;
-	for (FString SrcFilename : FilesToCopy)
-	{
-		// Update the slow task dialog
-		FFormatNamedArguments Args;
-		Args.Add(TEXT("SrcFilename"), FText::FromString(FPaths::GetCleanFilename(SrcFilename)));
-		SlowTask.EnterProgressFrame(1, FText::Format(LOCTEXT("CreatingProjectStatus_CopyingFile", "Copying File {SrcFilename}..."), Args));
-
-		FString FileRelPath = FPaths::GetPath(SrcFilename);
-		FPaths::MakePathRelativeTo(FileRelPath, *SrcFolder);
-
-		// Perform the copy. For file collisions, leave existing file.
-		const FString DestFilename = DestProjectFolder + TEXT("/") + FileRelPath + TEXT("/") + FPaths::GetCleanFilename(SrcFilename);
-		if (!FPaths::FileExists(DestFilename))
-		{
-			if (IFileManager::Get().Copy(*DestFilename, *SrcFilename, false) == COPY_OK)
-			{
-				CreatedFiles.Add(DestFilename);
-			}
-			else
-			{
-				FFormatNamedArguments FailArgs;
-				FailArgs.Add(TEXT("SrcFilename"), FText::FromString(SrcFilename));
-				FailArgs.Add(TEXT("DestFilename"), FText::FromString(DestFilename));
-				OutFailReason = FText::Format(LOCTEXT("FailedToCopyFile", "Failed to copy \"{SrcFilename}\" to \"{DestFilename}\"."), FailArgs);
-				DeleteCreatedFiles(DestProjectFolder, CreatedFiles);
-				return false;
-			}
-		}
-	}
-
-	return true;
+	FString const SrcFolder = FPaths::FeaturePackDir();
+	
+	IFileManager::Get().FindFilesRecursive(OutFilenames, *SrcFolder, TEXT("*.upack"), /*Files=*/true, /*Directories=*/false);
 }
 
-
 bool GameProjectUtils::CreateProject(const FProjectInformation& InProjectInfo, FText& OutFailReason)
 {
 	if ( !IsValidProjectFileForCreation(InProjectInfo.ProjectFilename, OutFailReason) )
@@ -965,6 +918,9 @@ bool GameProjectUtils::GenerateProjectFromScratch(const FProjectInformation& InP
 		return false;
 	}
 
+	// Insert any required feature packs (EG starter content) into ini file. These will be imported automatically when the editor is first run
+	InsertFeaturePacksIntoINIFile(InProjectInfo, OutFailReason);
+	
 	// Make the Content folder
 	const FString ContentFolder = NewProjectFolder / TEXT("Content");
 	if ( !IFileManager::Get().MakeDirectory(*ContentFolder) )
@@ -1044,17 +1000,6 @@ bool GameProjectUtils::GenerateProjectFromScratch(const FProjectInformation& InP
 
 	SlowTask.EnterProgressFrame();
 
-	if (InProjectInfo.bCopyStarterContent)
-	{
-		// Copy the starter content
-		if ( !CopyStarterContent(NewProjectFolder, OutFailReason) )
-		{
-			DeleteGeneratedProjectFiles(InProjectInfo.ProjectFilename);
-			DeleteCreatedFiles(NewProjectFolder, CreatedFiles);
-			return false;
-		}
-	}
-
 	UE_LOG(LogGameProjectGeneration, Log, TEXT("Created new project with %d files (plus project files)"), CreatedFiles.Num());
 	return true;
 }
@@ -1410,6 +1355,9 @@ bool GameProjectUtils::CreateProjectFromTemplate(const FProjectInformation& InPr
 		}
 	}
 
+	// Insert any required feature packs (EG starter content) into ini file. These will be imported automatically when the editor is first run
+	InsertFeaturePacksIntoINIFile(InProjectInfo, OutFailReason);
+	
 	SlowTask.EnterProgressFrame();
 
 	// Generate the project file
@@ -1477,17 +1425,6 @@ bool GameProjectUtils::CreateProjectFromTemplate(const FProjectInformation& InPr
 
 	SlowTask.EnterProgressFrame();
 
-	if (InProjectInfo.bCopyStarterContent)
-	{
-		// Copy the starter content
-		if ( !CopyStarterContent(DestFolder, OutFailReason) )
-		{
-			DeleteGeneratedProjectFiles(InProjectInfo.ProjectFilename);
-			DeleteCreatedFiles(DestFolder, CreatedFiles);
-			return false;
-		}
-	}
-
 	if (!TemplateDefs->PostGenerateProject(DestFolder, SrcFolder, InProjectInfo.ProjectFilename, InProjectInfo.TemplateFile, InProjectInfo.bShouldGenerateCode, OutFailReason))
 	{
 		DeleteGeneratedProjectFiles(InProjectInfo.ProjectFilename);
@@ -1669,66 +1606,23 @@ bool GameProjectUtils::GenerateConfigFiles(const FProjectInformation& InProjectI
 		
 		if (InProjectInfo.bCopyStarterContent)
 		{
-			FString StarterContentContentDir = FPaths::StarterContentDir() + TEXT("Content/");
-
-			TArray<FString> StarterContentMapFiles;
-			const FString FileWildcard = FString(TEXT("*")) + FPackageName::GetMapPackageExtension();
-		
 			FString SpecificEditorStartupMap;
-			FString SpecificGameDefaultMap;	
-			FString FullEditorStartupMapPath;
-			FString FullGameDefaultMapPath;
+			FString SpecificGameDefaultMap;
 
-			// First we check if there are maps specified in the DefaultEngine.ini in our starter content folder			
-			const FString StarterContentDefaultEngineIniFilename = FPaths::StarterContentDir() / TEXT("Config/DefaultEngine.ini");
-			if (FPaths::FileExists(StarterContentDefaultEngineIniFilename))
+			// If we have starter content packs available, specify starter map
+			if( IsStarterContentAvailableForNewProjects() == true )
 			{
-				FString StarterFileContents;
-				if (FFileHelper::LoadFileToString(StarterFileContents, *StarterContentDefaultEngineIniFilename))
+				if (InProjectInfo.TargetedHardware == EHardwareClass::Mobile)
 				{
-					TArray<FString> StarterIniLines;
-					StarterFileContents.ParseIntoArrayLines(&StarterIniLines);
-					for (int32 Line = 0; Line < StarterIniLines.Num();Line++)
-					{
-						FString EachLine = StarterIniLines[Line];
-						if (EachLine.StartsWith(TEXT("EditorStartupMap")))
-						{
-							EachLine.Split("=", nullptr, &SpecificEditorStartupMap);
-							FullEditorStartupMapPath = (StarterContentContentDir / SpecificEditorStartupMap) + FPackageName::GetMapPackageExtension();
-							FullEditorStartupMapPath = FullEditorStartupMapPath.Replace(TEXT("Game/"), TEXT(""));
-						}
-						if (EachLine.StartsWith(TEXT("GameDefaultMap")))
-						{
-							EachLine.Split("=", nullptr, &SpecificGameDefaultMap);
-							FullGameDefaultMapPath = (StarterContentContentDir / SpecificEditorStartupMap) + FPackageName::GetMapPackageExtension();
-							FullGameDefaultMapPath = FullGameDefaultMapPath.Replace(TEXT("Game/"), TEXT(""));
-						}
-					}					
+					SpecificEditorStartupMap = TEXT("/Game/MobileStarterContent/Maps/StarterMap");
+					SpecificGameDefaultMap = TEXT("/Game/MobileStarterContent/Maps/StarterMap");
 				}
-			}
-
-			// Look for maps in the content folder. If we don't specify maps for EditorStartup and GameDefault we will use the first we find in here
-			IFileManager::Get().FindFilesRecursive(StarterContentMapFiles, *FPaths::StarterContentDir(), *FileWildcard, /*Files=*/true, /*Directories=*/false);			
-			FString MapPackagePath;
-			if (StarterContentMapFiles.Num() > 0)
-			{
-				const FString BaseMapFilename = FPaths::GetBaseFilename(StarterContentMapFiles[0]);
-
-				FString MapPathRelToContent = FPaths::GetPath(StarterContentMapFiles[0]);
-				FPaths::MakePathRelativeTo(MapPathRelToContent, *StarterContentContentDir);
-
-				MapPackagePath = FString(TEXT("/Game/")) + MapPathRelToContent + TEXT("/") + BaseMapFilename;
-			}
-
-			// if either the files we specified don't exist or we didn't specify any, use the first map file we found in the content folder.
-			if (SpecificEditorStartupMap.IsEmpty() || FPaths::FileExists(FullEditorStartupMapPath) == false)
-			{
-				SpecificEditorStartupMap = MapPackagePath;
-			}
-			if (SpecificGameDefaultMap.IsEmpty() || FPaths::FileExists(FullGameDefaultMapPath) == false)
-			{
-				SpecificGameDefaultMap = MapPackagePath;
-			}
+				else
+				{
+					SpecificEditorStartupMap = TEXT("/Game/StarterContent/StarterMap");
+					SpecificGameDefaultMap = TEXT("/Game/StarterContent/StarterMap");
+				}
+			}						
 			
 			// Write out the settings for startup map and game default map
 			FileContents += TEXT("[/Script/EngineSettings.GameMapsSettings]") LINE_TERMINATOR;
@@ -1985,8 +1879,9 @@ bool GameProjectUtils::IsStarterContentAvailableForNewProjects()
 {
 	TArray<FString> StarterContentFiles;
 	GetStarterContentFiles(StarterContentFiles);
-
-	return (StarterContentFiles.Num() > 0);
+	
+	bool bHasStaterContent = StarterContentFiles.FindByPredicate([&](const FString& Str){ return Str.Contains("StarterContent"); }) != nullptr;
+	return bHasStaterContent;
 }
 
 TArray<FModuleContextInfo> GameProjectUtils::GetCurrentProjectModules()
@@ -3197,4 +3092,63 @@ void GameProjectUtils::HarvestCursorSyncLocation( FString& FinalOutput, FString&
 	FinalOutput = FinalOutput.Replace(TEXT("%CURSORFOCUSLOCATION%"), TEXT(""), ESearchCase::CaseSensitive);
 }
 
+bool GameProjectUtils::InsertFeaturePacksIntoINIFile(const FProjectInformation& InProjectInfo, FText& OutFailReason)
+{
+	bool bSuccessfullyProcessed = false;
+	const FString ProjectName = FPaths::GetBaseFilename(InProjectInfo.ProjectFilename);
+	const FString TemplateName = FPaths::GetBaseFilename(InProjectInfo.TemplateFile);
+	const FString SrcFolder = FPaths::GetPath(InProjectInfo.TemplateFile);
+	const FString DestFolder = FPaths::GetPath(InProjectInfo.ProjectFilename);
+
+	const FString ProjectConfigPath = DestFolder / TEXT("Config");
+	const FString IniFilename = ProjectConfigPath / TEXT("DefaultGame.ini");
+
+	TArray<FString> PackList;
+
+	// First the starter content
+	if (InProjectInfo.bCopyStarterContent)
+	{
+		FString StarterPack;
+		if (InProjectInfo.TargetedHardware == EHardwareClass::Mobile)
+		{
+			StarterPack = TEXT("InsertPack=(PackSource=\"MobileStarterContent.upack\",PackName=\"StarterContent\")");
+		}
+		else
+		{
+			StarterPack = TEXT("InsertPack=(PackSource=\"StarterContent.upack\",PackName=\"StarterContent\")");
+		}
+		PackList.Add(StarterPack);
+	}
+
+	// Now any packs specified in the template def.
+	UTemplateProjectDefs* TemplateDefs = LoadTemplateDefs(SrcFolder);
+	if (TemplateDefs != NULL)
+	{
+		for (int32 iPack = 0; iPack < TemplateDefs->PacksToInclude.Num(); ++iPack)
+		{
+			PackList.Add(TemplateDefs->PacksToInclude[iPack]);
+		}
+	}
+
+	FString FileOutput;
+	if (PackList.Num() != 0)
+	{
+		FileOutput = TEXT("[StartupActions]");
+		FileOutput += LINE_TERMINATOR;
+		FileOutput += TEXT("bAddPacks=True");
+		FileOutput += LINE_TERMINATOR;
+		for (int32 iLine = 0; iLine < PackList.Num(); ++iLine)
+		{
+			FileOutput += PackList[iLine] + LINE_TERMINATOR;
+
+		}
+		if (FFileHelper::SaveStringToFile(FileOutput, *IniFilename))
+		{
+			bSuccessfullyProcessed = true;
+		}
+	}
+
+	return bSuccessfullyProcessed;
+}
+
 #undef LOCTEXT_NAMESPACE
diff --git a/Engine/Source/Editor/GameProjectGeneration/Private/GameProjectUtils.h b/Engine/Source/Editor/GameProjectGeneration/Private/GameProjectUtils.h
--- a/Engine/Source/Editor/GameProjectGeneration/Private/GameProjectUtils.h
+++ b/Engine/Source/Editor/GameProjectGeneration/Private/GameProjectUtils.h
@@ -305,8 +305,8 @@ private:
 	/** Sets the engine association for a new project. Handles foreign and non-foreign projects. */
 	static bool SetEngineAssociationForForeignProject(const FString& ProjectFileName, FText& OutFailReason);
 
-	/** Copies starter content into the specified project folder. */
-	static bool CopyStarterContent(const FString& DestProjectFolder, FText& OutFailReason);
+	/** Insert any required feature packs into the DefaultGame.ini file */
+	static bool InsertFeaturePacksIntoINIFile(const FProjectInformation& InProjectInfo, FText& OutFailReason);
 
 	/** Returns list of starter content files */
 	static void GetStarterContentFiles(TArray<FString>& OutFilenames);
diff --git a/Engine/Source/Editor/GameProjectGeneration/Private/SNewProjectWizard.cpp b/Engine/Source/Editor/GameProjectGeneration/Private/SNewProjectWizard.cpp
--- a/Engine/Source/Editor/GameProjectGeneration/Private/SNewProjectWizard.cpp
+++ b/Engine/Source/Editor/GameProjectGeneration/Private/SNewProjectWizard.cpp
@@ -330,13 +330,22 @@ void SNewProjectWizard::Construct( const FArguments& InArgs )
 		TArray<SDecoratedEnumCombo<int32>::FComboOption> StarterContentInfo;
 		StarterContentInfo.Add(SDecoratedEnumCombo<int32>::FComboOption(
 			0, FSlateIcon(FEditorStyle::GetStyleSetName(), "GameProjectDialog.NoStarterContent"), LOCTEXT("NoStarterContent", "No Starter Content")));
-		StarterContentInfo.Add(SDecoratedEnumCombo<int32>::FComboOption(
-			1, FSlateIcon(FEditorStyle::GetStyleSetName(), "GameProjectDialog.IncludeStarterContent"), LOCTEXT("IncludeStarterContent", "With Starter Content")));
+		FString RequiredTooltipText = LOCTEXT("CopyStarterContent_ToolTip", "Enable to include an additional content pack containing simple placeable meshes with basic materials and textures.\nYou can opt out of including this to create a project that only has the bare essentials for the selected project template.").ToString();
+		// Only add the option to add starter content if its there to add !
+		if( GameProjectUtils::IsStarterContentAvailableForNewProjects() == true )
+		{
+			StarterContentInfo.Add(SDecoratedEnumCombo<int32>::FComboOption(
+				1, FSlateIcon(FEditorStyle::GetStyleSetName(), "GameProjectDialog.IncludeStarterContent"), LOCTEXT("IncludeStarterContent", "With Starter Content")));
+		}
+		else
+		{
+			RequiredTooltipText = LOCTEXT("NoStarterContent_ToolTip", "Starter content is not currently available.").ToString();
+		}
 
 		StartContentCombo = SNew(SDecoratedEnumCombo<int32>, MoveTemp(StarterContentInfo))
 			.SelectedEnum(this, &SNewProjectWizard::GetCopyStarterContentIndex)
 			.OnEnumChanged(this, &SNewProjectWizard::OnSetCopyStarterContent)
-			.ToolTipText( LOCTEXT("CopyStarterContent_ToolTip", "Enable to include an additional content pack containing simple placeable meshes with basic materials and textures.\nYou can opt out of including this to create a project that only has the bare essentials for the selected project template."));
+			.ToolTipText( RequiredTooltipText );
 	}
 
 	const float UniformPadding = 16.f;
@@ -836,11 +845,11 @@ FText SNewProjectWizard::GetStarterContentWarningTooltip() const
 {
 	if (SelectedGraphicsPreset == EGraphicsPreset::Maximum)
 	{
-		return LOCTEXT("StarterContentMobileWarning_Maximum", "Note: Starter content can increase the packaged size significantly, removing the example maps will result in only packaging content that is actually used");
+		return LOCTEXT("StarterContentMobileWarning_Maximum", "Note: Starter content will be inserted first time the project is opened, and can increase the packaged size significantly, removing the example maps will result in only packaging content that is actually used");
 	}
 	else
 	{
-		return LOCTEXT("StarterContentMobileWarning_Scalable", "Warning: Starter content is not optimized for scalable mobile projects");
+		return LOCTEXT("StarterContentMobileWarning_Scalable", "Warning: Starter content content will be inserted first time the project is opened, and is not optimized for scalable mobile projects");
 	}
 }
 
diff --git a/Engine/Source/Editor/UnrealEd/Private/FeaturePack/FeaturePack.cpp b/Engine/Source/Editor/UnrealEd/Private/FeaturePack/FeaturePack.cpp
new file mode 100644
--- /dev/null
+++ b/Engine/Source/Editor/UnrealEd/Private/FeaturePack/FeaturePack.cpp
@@ -0,0 +1,96 @@
+// Copyright 1998-2014 Epic Games, Inc. All Rights Reserved.
+
+#include "UnrealEd.h"
+#include "FeaturePack/FeaturePack.h"
+#include "AssetToolsModule.h"
+#include "AssetEditorManager.h"
+#include "AssetRegistryModule.h"
+#include "LevelEditor.h"
+#include "Toolkits/AssetEditorManager.h"
+
+
+DEFINE_LOG_CATEGORY_STATIC(LogFeaturePack, Log, All);
+
+FFeaturePack::FFeaturePack()
+{
+}
+
+void FFeaturePack::ImportPendingPacks()
+{	
+	bool bAddPacks;
+	if (GConfig->GetBool(TEXT("StartupActions"), TEXT("bAddPacks"), bAddPacks, GGameIni) == true)
+	{
+		if (bAddPacks == true)
+		{
+			ParsePacks();
+			UE_LOG(LogFeaturePack, Warning, TEXT("Inserted %d feature packs"), ImportedPacks.Num() );
+ 			GConfig->SetBool(TEXT("StartupActions"), TEXT("bAddPacks"), false, GGameIni);
+ 			GConfig->Flush(true, GGameIni);
+		}
+	}	
+}
+
+void FFeaturePack::ParsePacks()
+{
+	FAssetToolsModule& AssetToolsModule = FModuleManager::Get().LoadModuleChecked<FAssetToolsModule>("AssetTools");
+	// Look for pack insertions in the startup actions section
+	TArray<FString> PacksToAdd;
+	GConfig->GetArray(TEXT("StartupActions"), TEXT("InsertPack"), PacksToAdd, GGameIni);
+	for (int32 iPackEntry = 0; iPackEntry < PacksToAdd.Num(); iPackEntry++)
+	{
+		PackData EachPackData;
+		TArray<FString> PackEntries;
+		PacksToAdd[iPackEntry].ParseIntoArray(&PackEntries, TEXT(","), true);
+		FString PackSource;
+		FString PackName;
+		// Parse the pack name and source
+		for (int32 iEntry = 0; iEntry < PackEntries.Num(); iEntry++)
+		{
+			FString EachString = PackEntries[iEntry];
+			// String the parenthesis
+			EachString = EachString.Replace(TEXT("("), TEXT(""));
+			EachString = EachString.Replace(TEXT(")"), TEXT(""));
+			if (EachString.StartsWith(TEXT("PackSource=")) == true)
+			{
+				EachString = EachString.Replace(TEXT("PackSource="), TEXT(""));
+				EachString = EachString.TrimQuotes();
+				EachPackData.PackSource = EachString;
+			}
+			if (EachString.StartsWith(TEXT("PackName=")) == true)
+			{
+				EachString = EachString.Replace(TEXT("PackName="), TEXT(""));
+				EachString = EachString.TrimQuotes();
+				EachPackData.PackName = EachString;
+			}
+		}
+		
+		// If we found anything to insert, insert it !
+		if ((EachPackData.PackSource.IsEmpty() == false) && (EachPackData.PackName.IsEmpty() == false))
+		{
+			TArray<FString> EachImport;
+			FString FullPath = FPaths::FeaturePackDir() + EachPackData.PackSource;
+			EachImport.Add(FullPath);
+			EachPackData.ImportedObjects = AssetToolsModule.Get().ImportAssets(EachImport, EachPackData.PackName, nullptr);
+						
+			if (EachPackData.ImportedObjects.Num() == 0)
+			{			
+				UE_LOG(LogFeaturePack, Warning, TEXT("No objects imported installing pack %s"), *EachPackData.PackSource);
+			}
+			else
+			{
+				// Save any imported assets.
+				TArray<UPackage*> ToSave;
+				for (auto ImportedObject : EachPackData.ImportedObjects)
+				{
+					ToSave.AddUnique(ImportedObject->GetOutermost());
+				}
+				FEditorFileUtils::PromptForCheckoutAndSave(ToSave, /*bCheckDirty=*/ false, /*bPromptToSave=*/ false);				
+			}
+
+			ImportedPacks.Add(EachPackData);
+		}
+	}
+}
+
+
+#undef  LOCTEXT_NAMESPACE
diff --git a/Engine/Source/Editor/UnrealEd/Private/UnrealEdMisc.cpp b/Engine/Source/Editor/UnrealEd/Private/UnrealEdMisc.cpp
--- a/Engine/Source/Editor/UnrealEd/Private/UnrealEdMisc.cpp
+++ b/Engine/Source/Editor/UnrealEd/Private/UnrealEdMisc.cpp
@@ -256,6 +256,9 @@ void FUnrealEdMisc::OnInit()
 	{
 		bool bMapLoaded = false;
 
+		// Insert any feature packs if required. We need to do this before we try and load a map since any pack may contain a map
+		FeaturePackHandler.ImportPendingPacks();
+
 		FString ParsedMapName;
 		if ( FParse::Token(ParsedCmdLine, ParsedMapName, false) && 
 			 // If it's not a parameter
diff --git a/Engine/Source/Editor/UnrealEd/Public/FeaturePack/FeaturePack.h b/Engine/Source/Editor/UnrealEd/Public/FeaturePack/FeaturePack.h
new file mode 100644
--- /dev/null
+++ b/Engine/Source/Editor/UnrealEd/Public/FeaturePack/FeaturePack.h
@@ -0,0 +1,27 @@
+// Copyright 1998-2014 Epic Games, Inc. All Rights Reserved.
+
+#pragma once
+
+
+class FFeaturePack
+{
+public:
+	FFeaturePack();
+
+	/* Import any pending feature packs. */
+	void ImportPendingPacks();
+
+	/* Parse and insert any feature packs flagged to include and import them */
+	void ParsePacks();
+
+private:
+	struct PackData
+	{
+		FString PackSource;
+		FString PackName;
+		FString PackMap;
+		TArray<UObject*>	ImportedObjects;
+	};
+	/* List of packs we imported */
+	TArray<PackData>	ImportedPacks;
+};
diff --git a/Engine/Source/Editor/UnrealEd/Public/UnrealEd.h b/Engine/Source/Editor/UnrealEd/Public/UnrealEd.h
--- a/Engine/Source/Editor/UnrealEd/Public/UnrealEd.h
+++ b/Engine/Source/Editor/UnrealEd/Public/UnrealEd.h
@@ -30,7 +30,7 @@
 
 #include "MRUList.h"
 #include "Projects.h"
-
+#include "FeaturePack/FeaturePack.h"
 
 //#include "../Private/GeomFitUtils.h"
 
diff --git a/Engine/Source/Editor/UnrealEd/Public/UnrealEdMisc.h b/Engine/Source/Editor/UnrealEd/Public/UnrealEdMisc.h
--- a/Engine/Source/Editor/UnrealEd/Public/UnrealEdMisc.h
+++ b/Engine/Source/Editor/UnrealEd/Public/UnrealEdMisc.h
@@ -330,4 +330,7 @@ private:
 	FDelegateHandle OnMapChangedDelegateHandle;
 
 	FTimerHandle EditorAnalyticsHeartbeatTimerHandle;
+	
+	/* Handles insertion of feature packs (EG Starter content) at editor startup */
+	FFeaturePack FeaturePackHandler;
 };
diff --git a/Engine/Source/Runtime/Core/Private/Misc/Paths.cpp b/Engine/Source/Runtime/Core/Private/Misc/Paths.cpp
--- a/Engine/Source/Runtime/Core/Private/Misc/Paths.cpp
+++ b/Engine/Source/Runtime/Core/Private/Misc/Paths.cpp
@@ -358,9 +358,9 @@ FString FPaths::GameSourceDir()
 	return FPaths::GameDir() + TEXT("Source/");
 }
 
-FString FPaths::StarterContentDir()
+FString FPaths::FeaturePackDir()
 {
-	return FPaths::RootDir() + TEXT("Samples/StarterContent/");
+	return FPaths::RootDir() + TEXT("FeaturePacks/");
 }
 
 bool FPaths::IsProjectFilePathSet()
diff --git a/Engine/Source/Runtime/Core/Public/Misc/Paths.h b/Engine/Source/Runtime/Core/Public/Misc/Paths.h
--- a/Engine/Source/Runtime/Core/Public/Misc/Paths.h
+++ b/Engine/Source/Runtime/Core/Public/Misc/Paths.h
@@ -266,9 +266,9 @@ public:
 	static FString GameSourceDir();
 
 	/**
-	 * @return The directory where starter content is kept
+	 * @return The directory where feature packs are kept
 	 */
-	static FString StarterContentDir();
+	static FString FeaturePackDir();
 
 	/**
 	 * Checks whether the path to the project file, if any, is set.
