[INTEGRATE] Change 2431754 by [REDACTED_PERSON][REDACTED_HANDLE] on 2015/02/04 10:51:41

	Don't auto-open Visual Studio when compiling a new C++ Rocket project failed (that case was handled already)

diff --git a/Engine/Source/Editor/GameProjectGeneration/Private/SNewProjectWizard.cpp b/Engine/Source/Editor/GameProjectGeneration/Private/SNewProjectWizard.cpp
--- a/Engine/Source/Editor/GameProjectGeneration/Private/SNewProjectWizard.cpp
+++ b/Engine/Source/Editor/GameProjectGeneration/Private/SNewProjectWizard.cpp
@@ -1487,10 +1487,11 @@ void SNewProjectWizard::CreateAndOpenProject( )
 			    {
 					// Everything compiled OK, so we can go ahead and open the project
 					bCanOpenProject = true;
+		
+					// Open Visual Studio or Xcode if the user created a project with C++ files.  Note that if the code failed to compile, the
+					// BuildCodeProject() function will already offer to open the IDE for the user, so we only do this if every compiled OK.
+					OpenCodeIDE( ProjectFile );
 			    }
-
-				// Open Visual Studio or Xcode if the user created a project with C++ files.  Even if the project doesn't compile, we'll still open it in the IDE
-				OpenCodeIDE( ProjectFile );
 			}
 			else
 			{
