Changed tooltip to indicate which application will be used to open source files when clicking on their class's hyperlink.

diff --git a/Engine/Source/Editor/GameProjectGeneration/Private/SNewClassDialog.cpp b/Engine/Source/Editor/GameProjectGeneration/Private/SNewClassDialog.cpp
--- a/Engine/Source/Editor/GameProjectGeneration/Private/SNewClassDialog.cpp
+++ b/Engine/Source/Editor/GameProjectGeneration/Private/SNewClassDialog.cpp
@@ -392,7 +392,7 @@ void SNewClassDialog::Construct( const FArguments& InArgs )
 								.TextStyle(FEditorStyle::Get(), "DetailsView.GoToCodeHyperlinkStyle")
 								.OnNavigate(this, &SNewClassDialog::OnEditCodeClicked)
 								.Text(this, &SNewClassDialog::GetSelectedParentClassFilename)
-								.ToolTipText(LOCTEXT("GoToCode_ToolTip", "Click to open this source file in a text editor"))
+								.ToolTipText(FText::Format(LOCTEXT("GoToCode_ToolTip", "Click to open this source file in {0}"), FSourceCodeNavigation::GetSuggestedSourceCodeIDE()))
 								.Visibility(this, &SNewClassDialog::GetSourceHyperlinkVisibility)
 							]
 						]
diff --git a/Engine/Source/Editor/Kismet/Private/BlueprintEditor.cpp b/Engine/Source/Editor/Kismet/Private/BlueprintEditor.cpp
--- a/Engine/Source/Editor/Kismet/Private/BlueprintEditor.cpp
+++ b/Engine/Source/Editor/Kismet/Private/BlueprintEditor.cpp
@@ -1765,7 +1765,7 @@ void FBlueprintEditor::PostRegenerateMenusAndToolbars()
 				.Visibility(this, &FBlueprintEditor::GetNativeParentClassButtonsVisibility)
 				.OnNavigate(this, &FBlueprintEditor::OnEditParentClassNativeCodeClicked)
 				.Text(this, &FBlueprintEditor::GetTextForNativeParentClassHeaderLink)
-				.ToolTipText(LOCTEXT("GoToCode_ToolTip", "Click to open this source file in a text editor"))
+				.ToolTipText(FText::Format(LOCTEXT("GoToCode_ToolTip", "Click to open this source file in {0}"), FSourceCodeNavigation::GetSuggestedSourceCodeIDE()))
 			]
 			+ SHorizontalBox::Slot()
 			.AutoWidth()
diff --git a/Engine/Source/Editor/UnrealEd/Private/EditorClassUtils.cpp b/Engine/Source/Editor/UnrealEd/Private/EditorClassUtils.cpp
--- a/Engine/Source/Editor/UnrealEd/Private/EditorClassUtils.cpp
+++ b/Engine/Source/Editor/UnrealEd/Private/EditorClassUtils.cpp
@@ -133,7 +133,7 @@ TSharedRef<SWidget> FEditorClassUtils::GetSourceLinkFormatted(const UClass* Clas
 				.TextStyle(FEditorStyle::Get(), "DetailsView.GoToCodeHyperlinkStyle")
 				.OnNavigate_Static(&Local::OnEditCodeClicked, ClassHeaderPath)
 				.Text(FText::Format(CodeFormat, FText::FromString(FPaths::GetCleanFilename( *ClassHeaderPath ) ) ) )
-				.ToolTipText(NSLOCTEXT("SourceHyperlink", "GoToCode_ToolTip", "Click to open this source file in a text editor"));
+				.ToolTipText(FText::Format(NSLOCTEXT("SourceHyperlink", "GoToCode_ToolTip", "Click to open this source file in {0}"), FSourceCodeNavigation::GetSuggestedSourceCodeIDE()));
 		}
 	}
 
