[INTEGRATE] Change 2434517 by [REDACTED_PERSON][REDACTED_HANDLE] on 2015/02/05 19:13:26

	Starter content is now enabled by default and saves the user preference properly

diff --git a/Engine/Source/Editor/GameProjectGeneration/Private/SNewProjectWizard.cpp b/Engine/Source/Editor/GameProjectGeneration/Private/SNewProjectWizard.cpp
--- a/Engine/Source/Editor/GameProjectGeneration/Private/SNewProjectWizard.cpp
+++ b/Engine/Source/Editor/GameProjectGeneration/Private/SNewProjectWizard.cpp
@@ -299,7 +299,7 @@ void SNewProjectWizard::Construct( const FArguments& InArgs )
 	bLastGlobalValidityCheckSuccessful = true;
 	bLastNameAndLocationValidityCheckSuccessful = true;
 	bPreventPeriodicValidityChecksUntilNextChange = false;
-	bCopyStarterContent = false;
+	bCopyStarterContent = GEditor ? GEditor->AccessGameAgnosticSettings().bCopyStarterContentPreference : true;
 
 	IHardwareTargetingModule& HardwareTargeting = IHardwareTargetingModule::Get();
 
