[INTEGRATE] Change 2435332 by [REDACTED_PERSON][REDACTED_HANDLE].7 on 2015/02/06 09:25:07

	Generalized "Add Code to Project" dialog to allow creation of both c++ and blueprint class types.

	The add component option when blueprintable components is enabled, now opens up this dialog rather than the parent class picker.

	UE-8491 - IWCE: New C++ Component and New Blueprint Component should prompt you for a component type to subclass

diff --git a/Engine/Source/Developer/AssetTools/Private/AssetTypeActions/AssetTypeActions_Class.cpp b/Engine/Source/Developer/AssetTools/Private/AssetTypeActions/AssetTypeActions_Class.cpp
--- a/Engine/Source/Developer/AssetTools/Private/AssetTypeActions/AssetTypeActions_Class.cpp
+++ b/Engine/Source/Developer/AssetTools/Private/AssetTypeActions/AssetTypeActions_Class.cpp
@@ -30,7 +30,12 @@ void FAssetTypeActions_Class::GetActions(const TArray<UObject*>& InObjects, FMen
 			BaseClassPath = FPaths::GetPath(BaseClassPath);
 		}
 
-		FGameProjectGenerationModule::Get().OpenAddCodeToProjectDialog(BaseClass, BaseClassPath, FGlobalTabmanager::Get()->GetRootWindow());
+		FGameProjectGenerationModule::Get().OpenAddCodeToProjectDialog(
+			FAddToProjectConfig()
+			.ParentClass(BaseClass)
+			.InitialPath(BaseClassPath)
+			.ParentWindow(FGlobalTabmanager::Get()->GetRootWindow())
+		);
 	};
 
 	auto CanCreateDerivedClass = [bIsValidBaseClass]() -> bool
diff --git a/Engine/Source/Editor/ClassViewer/Private/SClassViewer.cpp b/Engine/Source/Editor/ClassViewer/Private/SClassViewer.cpp
--- a/Engine/Source/Editor/ClassViewer/Private/SClassViewer.cpp
+++ b/Engine/Source/Editor/ClassViewer/Private/SClassViewer.cpp
@@ -1138,7 +1138,11 @@ namespace ClassViewer
 		 */
 		static void OpenCreateCPlusPlusClassWizard(UClass* InCreationClass)
 		{
-			FGameProjectGenerationModule::Get().OpenAddCodeToProjectDialog(InCreationClass, FString(), FGlobalTabmanager::Get()->GetRootWindow());
+			FGameProjectGenerationModule::Get().OpenAddCodeToProjectDialog(
+				FAddToProjectConfig()
+				.ParentClass(InCreationClass)
+				.ParentWindow(FGlobalTabmanager::Get()->GetRootWindow())
+			);
 		}
 
 		/**
diff --git a/Engine/Source/Editor/ContentBrowser/Private/SContentBrowser.cpp b/Engine/Source/Editor/ContentBrowser/Private/SContentBrowser.cpp
--- a/Engine/Source/Editor/ContentBrowser/Private/SContentBrowser.cpp
+++ b/Engine/Source/Editor/ContentBrowser/Private/SContentBrowser.cpp
@@ -1167,7 +1167,11 @@ void SContentBrowser::NewClassRequested(const FString& SelectedPath)
 		NativeClassHierarchy->GetFileSystemPath(SelectedPath, ExistingFolderPath);
 	}
 
-	FGameProjectGenerationModule::Get().OpenAddCodeToProjectDialog(nullptr, ExistingFolderPath, FGlobalTabmanager::Get()->GetRootWindow());
+	FGameProjectGenerationModule::Get().OpenAddCodeToProjectDialog(
+		FAddToProjectConfig()
+		.InitialPath(ExistingFolderPath)
+		.ParentWindow(FGlobalTabmanager::Get()->GetRootWindow())
+	);
 }
 
 void SContentBrowser::NewFolderRequested(const FString& SelectedPath)
diff --git a/Engine/Source/Editor/GameProjectGeneration/GameProjectGeneration.Build.cs b/Engine/Source/Editor/GameProjectGeneration/GameProjectGeneration.Build.cs
--- a/Engine/Source/Editor/GameProjectGeneration/GameProjectGeneration.Build.cs
+++ b/Engine/Source/Editor/GameProjectGeneration/GameProjectGeneration.Build.cs
@@ -33,6 +33,7 @@ public class GameProjectGeneration : ModuleRules
 				"UnrealEd",
 				"DesktopPlatform",
                 "HardwareTargeting",
+                "AssetRegistry"
 			}
 		);
 
diff --git a/Engine/Source/Editor/GameProjectGeneration/Private/GameProjectGenerationModule.cpp b/Engine/Source/Editor/GameProjectGeneration/Private/GameProjectGenerationModule.cpp
--- a/Engine/Source/Editor/GameProjectGeneration/Private/GameProjectGenerationModule.cpp
+++ b/Engine/Source/Editor/GameProjectGeneration/Private/GameProjectGenerationModule.cpp
@@ -6,6 +6,8 @@
 #include "TemplateCategory.h"
 #include "SourceCodeNavigation.h"
 
+#include "MainFrame.h"
+
 IMPLEMENT_MODULE( FGameProjectGenerationModule, GameProjectGeneration );
 DEFINE_LOG_CATEGORY(LogGameProjectGeneration);
 
@@ -54,17 +56,15 @@ TSharedRef<class SWidget> FGameProjectGenerationModule::CreateNewClassDialog(con
 	return SNew(SNewClassDialog).Class(InClass);
 }
 
-
-void FGameProjectGenerationModule::OpenAddCodeToProjectDialog(const TSharedPtr<SWindow>& InParentWindow)
+void FGameProjectGenerationModule::OpenAddCodeToProjectDialog(const FAddToProjectConfig& Config)
 {
-	GameProjectUtils::OpenAddCodeToProjectDialog(nullptr, FString(), InParentWindow);
+	GameProjectUtils::OpenAddToProjectDialog(Config, EClassDomain::Native);
 	AddCodeToProjectDialogOpenedEvent.Broadcast();
 }
 
-void FGameProjectGenerationModule::OpenAddCodeToProjectDialog(const UClass* InClass, const FString& InInitialPath, const TSharedPtr<SWindow>& InParentWindow, bool bModal, FOnCodeAddedToProject OnCodeAddedToProject, const FString InDefaultClassPrefix, const FString InDefaultClassName)
+void FGameProjectGenerationModule::OpenAddBlueprintToProjectDialog(const FAddToProjectConfig& Config)
 {
-	GameProjectUtils::OpenAddCodeToProjectDialog(InClass, InInitialPath, InParentWindow, bModal, OnCodeAddedToProject, InDefaultClassPrefix, InDefaultClassName);
-	AddCodeToProjectDialogOpenedEvent.Broadcast();
+	GameProjectUtils::OpenAddToProjectDialog(Config, EClassDomain::Blueprint);
 }
 
 void FGameProjectGenerationModule::TryMakeProjectFileWriteable(const FString& ProjectFile)
diff --git a/Engine/Source/Editor/GameProjectGeneration/Private/GameProjectUtils.cpp b/Engine/Source/Editor/GameProjectGeneration/Private/GameProjectUtils.cpp
--- a/Engine/Source/Editor/GameProjectGeneration/Private/GameProjectUtils.cpp
+++ b/Engine/Source/Editor/GameProjectGeneration/Private/GameProjectUtils.cpp
@@ -2,6 +2,9 @@
 
 
 #include "GameProjectGenerationPrivatePCH.h"
+
+#include "FeaturedClasses.inl"
+
 #include "UnrealEdMisc.h"
 #include "ISourceControlModule.h"
 #include "MainFrame.h"
@@ -38,7 +41,7 @@ static_assert(PLATFORM_MAX_FILEPATH_LENGTH - MAX_PROJECT_PATH_BUFFER_SPACE > 0,
 TWeakPtr<SNotificationItem> GameProjectUtils::UpdateGameProjectNotification = NULL;
 TWeakPtr<SNotificationItem> GameProjectUtils::WarningProjectNameNotification = NULL;
 
-FString GameProjectUtils::FNewClassInfo::GetClassName() const
+FString FNewClassInfo::GetClassName() const
 {
 	switch(ClassType)
 	{
@@ -60,7 +63,7 @@ FString GameProjectUtils::FNewClassInfo::GetClassName() const
 	return TEXT("");
 }
 
-FString GameProjectUtils::FNewClassInfo::GetClassDescription() const
+FString FNewClassInfo::GetClassDescription() const
 {
 	switch(ClassType)
 	{
@@ -98,13 +101,13 @@ FString GameProjectUtils::FNewClassInfo::GetClassDescription() const
 	return TEXT("");
 }
 
-const FSlateBrush* GameProjectUtils::FNewClassInfo::GetClassIcon() const
+const FSlateBrush* FNewClassInfo::GetClassIcon() const
 {
 	// Safe to do even if BaseClass is null, since FindIconForClass will return the default icon
 	return FClassIconFinder::FindIconForClass(BaseClass);
 }
 
-FString GameProjectUtils::FNewClassInfo::GetClassPrefixCPP() const
+FString FNewClassInfo::GetClassPrefixCPP() const
 {
 	switch(ClassType)
 	{
@@ -126,7 +129,7 @@ FString GameProjectUtils::FNewClassInfo::GetClassPrefixCPP() const
 	return TEXT("");
 }
 
-FString GameProjectUtils::FNewClassInfo::GetClassNameCPP() const
+FString FNewClassInfo::GetClassNameCPP() const
 {
 	switch(ClassType)
 	{
@@ -148,7 +151,7 @@ FString GameProjectUtils::FNewClassInfo::GetClassNameCPP() const
 	return TEXT("");
 }
 
-FString GameProjectUtils::FNewClassInfo::GetCleanClassName(const FString& ClassName) const
+FString FNewClassInfo::GetCleanClassName(const FString& ClassName) const
 {
 	FString CleanClassName = ClassName;
 
@@ -176,7 +179,7 @@ FString GameProjectUtils::FNewClassInfo::GetCleanClassName(const FString& ClassN
 	return CleanClassName;
 }
 
-FString GameProjectUtils::FNewClassInfo::GetFinalClassName(const FString& ClassName) const
+FString FNewClassInfo::GetFinalClassName(const FString& ClassName) const
 {
 	const FString CleanClassName = GetCleanClassName(ClassName);
 
@@ -192,7 +195,7 @@ FString GameProjectUtils::FNewClassInfo::GetFinalClassName(const FString& ClassN
 	return CleanClassName;
 }
 
-bool GameProjectUtils::FNewClassInfo::GetIncludePath(FString& OutIncludePath) const
+bool FNewClassInfo::GetIncludePath(FString& OutIncludePath) const
 {
 	switch(ClassType)
 	{
@@ -210,7 +213,7 @@ bool GameProjectUtils::FNewClassInfo::GetIncludePath(FString& OutIncludePath) co
 	return false;
 }
 
-FString GameProjectUtils::FNewClassInfo::GetHeaderFilename(const FString& ClassName) const
+FString FNewClassInfo::GetHeaderFilename(const FString& ClassName) const
 {
 	const FString HeaderFilename = GetFinalClassName(ClassName) + TEXT(".h");
 
@@ -226,7 +229,7 @@ FString GameProjectUtils::FNewClassInfo::GetHeaderFilename(const FString& ClassN
 	return HeaderFilename;
 }
 
-FString GameProjectUtils::FNewClassInfo::GetSourceFilename(const FString& ClassName) const
+FString FNewClassInfo::GetSourceFilename(const FString& ClassName) const
 {
 	const FString SourceFilename = GetFinalClassName(ClassName) + TEXT(".cpp");
 
@@ -242,7 +245,7 @@ FString GameProjectUtils::FNewClassInfo::GetSourceFilename(const FString& ClassN
 	return SourceFilename;
 }
 
-FString GameProjectUtils::FNewClassInfo::GetHeaderTemplateFilename() const
+FString FNewClassInfo::GetHeaderTemplateFilename() const
 {
 	switch(ClassType)
 	{
@@ -278,7 +281,7 @@ FString GameProjectUtils::FNewClassInfo::GetHeaderTemplateFilename() const
 	return TEXT("");
 }
 
-FString GameProjectUtils::FNewClassInfo::GetSourceTemplateFilename() const
+FString FNewClassInfo::GetSourceTemplateFilename() const
 {
 	switch(ClassType)
 	{
@@ -674,29 +677,38 @@ bool GameProjectUtils::UpdateGameProject(const FString& ProjectFile, const FStri
 	return UpdateGameProjectFile(ProjectFile, EngineIdentifier, NULL, OutFailReason);
 }
 
-void GameProjectUtils::OpenAddCodeToProjectDialog(const UClass* InClass, const FString& InInitialPath, const TSharedPtr<SWindow>& InParentWindow, bool bModal, FOnCodeAddedToProject OnCodeAddedToProject, const FString InDefaultClassPrefix, const FString InDefaultClassName )
+void GameProjectUtils::OpenAddToProjectDialog(const FAddToProjectConfig& Config, EClassDomain InDomain)
 {
 	// If we've been given a class then we only show the second page of the dialog, so we can make the window smaller as that page doesn't have as much content
-	const FVector2D WindowSize = (InClass) ? FVector2D(940, 380) : FVector2D(940, 540);
+	const FVector2D WindowSize = (Config._ParentClass) ? FVector2D(940, 380) : FVector2D(940, 540);
+
+	FText WindowTitle = Config._WindowTitle;
+	if (WindowTitle.IsEmpty())
+	{
+		WindowTitle = InDomain == EClassDomain::Native ? LOCTEXT("AddCodeWindowHeader_Native", "Add Code") : LOCTEXT("AddCodeWindowHeader_Blueprint", "Add Blueprint");
+	}
 
 	TSharedRef<SWindow> AddCodeWindow =
 		SNew(SWindow)
-		.Title(LOCTEXT( "AddCodeWindowHeader", "Add Code"))
+		.Title( WindowTitle )
 		.ClientSize( WindowSize )
 		.SizingRule( ESizingRule::FixedSize )
 		.SupportsMinimize(false) .SupportsMaximize(false);
 
 	TSharedRef<SNewClassDialog> NewClassDialog = 
 		SNew(SNewClassDialog)
-		.Class(InClass)
-		.InitialPath(InInitialPath)
-		.OnCodeAddedToProject( OnCodeAddedToProject )
-		.DefaultClassPrefix( InDefaultClassPrefix )
-		.DefaultClassName( InDefaultClassName );
+		.Class(Config._ParentClass)
+		.ClassViewerFilter(Config._AllowableParents)
+		.ClassDomain(InDomain)
+		.FeaturedClasses(Config._FeaturedClasses)
+		.InitialPath(Config._InitialPath)
+		.OnAddedToProject( Config._OnAddedToProject )
+		.DefaultClassPrefix( Config._DefaultClassPrefix )
+		.DefaultClassName( Config._DefaultClassName );
 
 	AddCodeWindow->SetContent( NewClassDialog );
 
-	TSharedPtr<SWindow> ParentWindow = InParentWindow;
+	TSharedPtr<SWindow> ParentWindow = Config._ParentWindow;
 	if (!ParentWindow.IsValid())
 	{
 		static const FName MainFrameModuleName = "MainFrame";
@@ -704,32 +716,21 @@ void GameProjectUtils::OpenAddCodeToProjectDialog(const UClass* InClass, const F
 		ParentWindow = MainFrameModule.GetParentWindow();
 	}
 
-	if (ParentWindow.IsValid())
+	if (Config._bModal)
 	{
-		if( bModal )
-		{
-			FSlateApplication::Get().AddModalWindow(AddCodeWindow, ParentWindow);
-		}
-		else
-		{
-			FSlateApplication::Get().AddWindowAsNativeChild(AddCodeWindow, ParentWindow.ToSharedRef());
-		}
+		FSlateApplication::Get().AddModalWindow(AddCodeWindow, ParentWindow);
+	}
+	else if (ParentWindow.IsValid())
+	{
+		FSlateApplication::Get().AddWindowAsNativeChild(AddCodeWindow, ParentWindow.ToSharedRef());
 	}
 	else
 	{
-		if(bModal)
-		{
-			FSlateApplication::Get().AddModalWindow(AddCodeWindow, nullptr);
-		}
-		else
-		{
-			FSlateApplication::Get().AddWindow(AddCodeWindow);
-		}
-
+		FSlateApplication::Get().AddWindow(AddCodeWindow);
 	}
 }
 
-bool GameProjectUtils::IsValidClassNameForCreation(const FString& NewClassName, const FModuleContextInfo& ModuleInfo, const TSet<FString>& DisallowedHeaderNames, FText& OutFailReason)
+bool GameProjectUtils::IsValidClassNameForCreation(const FString& NewClassName, FText& OutFailReason)
 {
 	if ( NewClassName.IsEmpty() )
 	{
@@ -764,6 +765,16 @@ bool GameProjectUtils::IsValidClassNameForCreation(const FString& NewClassName,
 		return false;
 	}
 
+	return true;
+}
+
+bool GameProjectUtils::IsValidClassNameForCreation(const FString& NewClassName, const FModuleContextInfo& ModuleInfo, const TSet<FString>& DisallowedHeaderNames, FText& OutFailReason)
+{
+	if (!IsValidClassNameForCreation(NewClassName, OutFailReason))
+	{
+		return false;
+	}
+	
 	// Look for a duplicate class in memory
 	for ( TObjectIterator<UClass> ClassIt; ClassIt; ++ClassIt )
 	{
diff --git a/Engine/Source/Editor/GameProjectGeneration/Private/GameProjectUtils.h b/Engine/Source/Editor/GameProjectGeneration/Private/GameProjectUtils.h
--- a/Engine/Source/Editor/GameProjectGeneration/Private/GameProjectUtils.h
+++ b/Engine/Source/Editor/GameProjectGeneration/Private/GameProjectUtils.h
@@ -5,6 +5,8 @@
 #include "HardwareTargetingModule.h"
 
 struct FModuleContextInfo;
+struct FAddToProjectConfig;
+enum class EClassDomain : uint8;
 
 struct FProjectInformation
 {
@@ -46,95 +48,6 @@ public:
 		Classes,
 	};
 
-	/** Information used when creating a new class via AddCodeToProject */
-	struct FNewClassInfo
-	{
-		/** The type of class we want to create */
-		enum class EClassType : uint8
-		{
-			/** The new class is using a UObject as a base, consult BaseClass for the type */
-			UObject,
-
-			/** The new class should be an empty standard C++ class */
-			EmptyCpp,
-
-			/** The new class should be a Slate widget, deriving from SCompoundWidget */
-			SlateWidget,
-
-			/** The new class should be a Slate widget style, deriving from FSlateWidgetStyle, along with its associated UObject wrapper class */
-			SlateWidgetStyle,
-		};
-
-		/** Default constructor; must produce an object which fails the IsSet check */
-		FNewClassInfo()
-			: ClassType(EClassType::UObject)
-			, BaseClass(nullptr)
-		{
-		}
-
-		/** Convenience constructor so you can construct from a EClassType */
-		explicit FNewClassInfo(const EClassType InClassType)
-			: ClassType(InClassType)
-			, BaseClass(nullptr)
-		{
-		}
-
-		/** Convenience constructor so you can construct from a UClass */
-		explicit FNewClassInfo(const UClass* InBaseClass)
-			: ClassType(EClassType::UObject)
-			, BaseClass(InBaseClass)
-		{
-		}
-
-		/** Check to see if this struct is set to something that could be used to create a new class */
-		bool IsSet() const
-		{
-			return ClassType != EClassType::UObject || BaseClass;
-		}
-
-		/** Get the "friendly" class name to use in the UI */
-		FString GetClassName() const;
-
-		/** Get the class description to use in the UI */
-		FString GetClassDescription() const;
-
-		/** Get the class icon to use in the UI */
-		const FSlateBrush* GetClassIcon() const;
-
-		/** Get the C++ prefix used for this class type */
-		FString GetClassPrefixCPP() const;
-
-		/** Get the C++ class name; this may or may not be prefixed, but will always produce a valid C++ name via GetClassPrefix() + GetClassName() */
-		FString GetClassNameCPP() const;
-
-		/** Some classes may apply a particular suffix; this function returns the class name with those suffixes removed */
-		FString GetCleanClassName(const FString& ClassName) const;
-
-		/** Some classes may apply a particular suffix; this function returns the class name that will ultimately be used should that happen */
-		FString GetFinalClassName(const FString& ClassName) const;
-
-		/** Get the path needed to include this class into another file */
-		bool GetIncludePath(FString& OutIncludePath) const;
-
-		/** Given a class name, generate the header file (.h) that should be used for this class */
-		FString GetHeaderFilename(const FString& ClassName) const;
-
-		/** Given a class name, generate the source file (.cpp) that should be used for this class */
-		FString GetSourceFilename(const FString& ClassName) const;
-
-		/** Get the generation template filename to used based on the current class type */
-		FString GetHeaderTemplateFilename() const;
-
-		/** Get the generation template filename to used based on the current class type */
-		FString GetSourceTemplateFilename() const;
-
-		/** The type of class we want to create */
-		EClassType ClassType;
-
-		/** Base class information; if the ClassType is UObject */
-		const UClass* BaseClass;
-	};
-
 	/** Used as a function return result when a project is duplicated when upgrading project's version in Convert project dialog - Open a copy */
 	enum class EProjectDuplicateResult : uint8
 	{
@@ -173,19 +86,17 @@ public:
 	static bool UpdateGameProject(const FString& ProjectFile, const FString& EngineIdentifier, FText& OutFailReason);
 
 	/** 
-	 * Opens a dialog to add code files to the current project. 
+	 * Opens a dialog to add code files or blueprints to the current project. 
 	 *
-	 * @param	InClass			The class we should force the user to use as their base class type, or null to allow the user to choose their base class in the UI
-	 * @param	InInitialPath		The initial path we should use as the destination for the new class header file, or an empty string to choose a suitable default based upon the module path
-	 * @param	InParentWindow		The parent window the dialog should use, or null to choose a suitable default parent window (the main frame, if available)
-	 * @param	bModal			True if the window should be modal and force the user to make a decision before continuing or false to let the user proceed with other tasks while the window is open
-	 * @param	OnCodeAddedToProject	Callback for when code is successfully added to the project
-	 * @param	InDefaultClassPrefix	Optional argument that specifies the prefix for the new class name.  The user will be able to type their own name if they don't like this name.  Defaults to "My" if not specified or empty
-	 * @param	InDefaultClassName	Optional argument that specifies the default name for the new class being added.  The user will be able to type their own name if they don't like this name.  If empty, defaults to the name of the inherited class.
+	 * @param	Config			Configuration options for the dialog
+	 * @param	InDomain		The domain of the class we're creating (native or blueprint)
 	 */
-	static void OpenAddCodeToProjectDialog(const UClass* InClass, const FString& InInitialPath, const TSharedPtr<SWindow>& InParentWindow, bool bModal = false, FOnCodeAddedToProject OnCodeAddedToProject = FOnCodeAddedToProject(), const FString InDefaultClassPrefix = FString(), const FString InDefaultClassName = FString() );
+	static void OpenAddToProjectDialog(const FAddToProjectConfig& Config, EClassDomain InDomain);
 
 	/** Returns true if the specified class name is properly formed and does not conflict with another class */
+	static bool IsValidClassNameForCreation(const FString& NewClassName, FText& OutFailReason);
+
+	/** Returns true if the specified class name is properly formed and does not conflict with another class, including source/header files */
 	static bool IsValidClassNameForCreation(const FString& NewClassName, const FModuleContextInfo& ModuleInfo, const TSet<FString>& DisallowedHeaderNames, FText& OutFailReason);
 
 	/** Returns true if the specified class is a valid base class for the given module */
diff --git a/Engine/Source/Editor/GameProjectGeneration/Private/SNewClassDialog.cpp b/Engine/Source/Editor/GameProjectGeneration/Private/SNewClassDialog.cpp
--- a/Engine/Source/Editor/GameProjectGeneration/Private/SNewClassDialog.cpp
+++ b/Engine/Source/Editor/GameProjectGeneration/Private/SNewClassDialog.cpp
@@ -5,6 +5,7 @@
 #include "SourceCodeNavigation.h"
 #include "ClassViewerModule.h"
 #include "ClassViewerFilter.h"
+#include "Editor/ContentBrowser/Public/ContentBrowserModule.h"
 #include "Editor/ClassViewer/Private/SClassViewer.h"
 #include "DesktopPlatformModule.h"
 #include "Editor/Documentation/Public/IDocumentation.h"
@@ -23,6 +24,8 @@
 #include "GameFramework/HUD.h"
 #include "GameFramework/PlayerState.h"
 #include "Kismet/BlueprintFunctionLibrary.h"
+#include "Kismet2/KismetEditorUtilities.h"
+#include "AssetRegistryModule.h"
 #include "SNotificationList.h"
 #include "NotificationManager.h"
 
@@ -32,9 +35,9 @@
 
 struct FParentClassItem
 {
-	GameProjectUtils::FNewClassInfo ParentClassInfo;
+	FNewClassInfo ParentClassInfo;
 
-	FParentClassItem(const GameProjectUtils::FNewClassInfo& InParentClassInfo)
+	FParentClassItem(const FNewClassInfo& InParentClassInfo)
 		: ParentClassInfo(InParentClassInfo)
 	{}
 };
@@ -77,6 +80,8 @@ static void FindPublicEngineHeaderFiles(TArray<FString>& OutFiles, const FString
 BEGIN_SLATE_FUNCTION_BUILD_OPTIMIZATION
 void SNewClassDialog::Construct( const FArguments& InArgs )
 {
+	ClassDomain = InArgs._ClassDomain;
+
 	{
 		TArray<FModuleContextInfo> CurrentModules = GameProjectUtils::GetCurrentProjectModules();
 		check(CurrentModules.Num()); // this should never happen since GetCurrentProjectModules is supposed to add a dummy runtime module if the project currently has no modules
@@ -89,7 +94,12 @@ void SNewClassDialog::Construct( const FArguments& InArgs )
 	}
 
 	// If we've been given an initial path that maps to a valid project module, use that as our initial module and path
-	if(!InArgs._InitialPath.IsEmpty())
+
+	if (ClassDomain == EClassDomain::Blueprint)
+	{
+		NewClassPath = InArgs._InitialPath.IsEmpty() ? TEXT("/Game") : InArgs._InitialPath;
+	}
+	else if(!InArgs._InitialPath.IsEmpty())
 	{
 		const FString AbsoluteInitialPath = FPaths::ConvertRelativePathToFull(InArgs._InitialPath);
 		for(const auto& AvailableModule : AvailableModules)
@@ -109,7 +119,7 @@ void SNewClassDialog::Construct( const FArguments& InArgs )
 	// If we didn't get given a valid path override (see above), try and automatically work out the best default module
 	// If we have a runtime module with the same name as our project, then use that
 	// Otherwise, set out default target module as the first runtime module in the list
-	if(!SelectedModuleInfo.IsValid())
+	if(ClassDomain == EClassDomain::Native && !SelectedModuleInfo.IsValid())
 	{
 		const FString ProjectName = FApp::GetGameName();
 		for(const auto& AvailableModule : AvailableModules)
@@ -138,7 +148,7 @@ void SNewClassDialog::Construct( const FArguments& InArgs )
 
 	ClassLocation = GameProjectUtils::EClassLocation::UserDefined; // the first call to UpdateInputValidity will set this correctly based on NewClassPath
 
-	ParentClassInfo = GameProjectUtils::FNewClassInfo(InArgs._Class);
+	ParentClassInfo = FNewClassInfo(InArgs._Class);
 
 	bShowFullClassTree = false;
 
@@ -147,9 +157,6 @@ void SNewClassDialog::Construct( const FArguments& InArgs )
 	bLastInputValidityCheckSuccessful = true;
 	bPreventPeriodicValidityChecksUntilNextChange = false;
 
-	SetupParentClassItems();
-	UpdateInputValidity();
-
 	FClassViewerInitializationOptions Options;
 	Options.Mode = EClassViewerMode::ClassPicker;
 	Options.DisplayMode = EClassViewerDisplayMode::TreeView;
@@ -160,17 +167,39 @@ void SNewClassDialog::Construct( const FArguments& InArgs )
 	Options.bShowNoneOption = false;
 	Options.bShowObjectRootClass = true;
 
-	// Prevent creating native classes based on blueprint classes
-	Options.ClassFilter = MakeShareable(new FNativeClassParentFilter());
+	if (InArgs._ClassViewerFilter.IsValid())
+	{
+		Options.ClassFilter = InArgs._ClassViewerFilter;
+	}
+	else if (InArgs._ClassDomain == EClassDomain::Native)
+	{
+		// Prevent creating native classes based on blueprint classes
+		Options.ClassFilter = MakeShareable(new FNativeClassParentFilter());
+	}
 
 	ClassViewer = StaticCastSharedRef<SClassViewer>(FModuleManager::LoadModuleChecked<FClassViewerModule>("ClassViewer").CreateClassViewer(Options, FOnClassPicked::CreateSP(this, &SNewClassDialog::OnAdvancedClassSelected)));
 
+	SetupParentClassItems(InArgs._FeaturedClasses);
+	UpdateInputValidity();
+
 	TSharedRef<SWidget> DocWidget = IDocumentation::Get()->CreateAnchor(TAttribute<FString>(this, &SNewClassDialog::GetSelectedParentDocLink));
 	DocWidget->SetVisibility(TAttribute<EVisibility>(this, &SNewClassDialog::GetDocLinkVisibility));
 
 	const float EditableTextHeight = 26.0f;
 
-	OnCodeAddedToProject = InArgs._OnCodeAddedToProject;
+	IContentBrowserSingleton& ContentBrowser = FModuleManager::LoadModuleChecked<FContentBrowserModule>("ContentBrowser").Get();
+
+	FPathPickerConfig BlueprintPathConfig;
+	if (ClassDomain == EClassDomain::Blueprint)
+	{
+		BlueprintPathConfig.DefaultPath = InArgs._InitialPath;
+		BlueprintPathConfig.bFocusSearchBoxWhenOpened = false;
+		BlueprintPathConfig.bAllowContextMenu = false;
+		BlueprintPathConfig.bAllowClassesFolder = false;
+		BlueprintPathConfig.OnPathSelected = FOnPathSelected::CreateSP(this, &SNewClassDialog::OnBlueprintPathSelected);
+	}
+
+	OnAddedToProject = InArgs._OnAddedToProject;
 
 	ChildSlot
 	[
@@ -186,8 +215,12 @@ void SNewClassDialog::Construct( const FArguments& InArgs )
 				SAssignNew( MainWizard, SWizard)
 				.ShowPageList(false)
 				.CanFinish(this, &SNewClassDialog::CanFinish)
-				.FinishButtonText( LOCTEXT("FinishButtonText", "Create Class") )
-				.FinishButtonToolTip ( LOCTEXT("FinishButtonToolTip", "Creates the code files to add your new class.") )
+				.FinishButtonText( ClassDomain == EClassDomain::Native ? LOCTEXT("FinishButtonText_Native", "Create Class") : LOCTEXT("FinishButtonText_Blueprint", "Create Blueprint Class") )
+				.FinishButtonToolTip (
+					ClassDomain == EClassDomain::Native ?
+					LOCTEXT("FinishButtonToolTip_Native", "Creates the code files to add your new class.") : 
+					LOCTEXT("FinishButtonToolTip_Blueprint", "Creates the new Blueprint class based on the specified parent class.")
+					)
 				.OnCanceled(this, &SNewClassDialog::CancelClicked)
 				.OnFinished(this, &SNewClassDialog::FinishClicked)
 				.InitialPageIndex(ParentClassInfo.IsSet() ? 1 : 0)
@@ -266,7 +299,11 @@ void SNewClassDialog::Construct( const FArguments& InArgs )
 						.VAlign(VAlign_Center)
 						[
 							SNew(STextBlock)
-							.Text(LOCTEXT("ChooseParentClassDescription", "This will add a C++ header and source code file to your game project."))
+							.Text(
+								ClassDomain == EClassDomain::Native ?
+								LOCTEXT("ChooseParentClassDescription_Native", "This will add a C++ header and source code file to your game project.") :
+								LOCTEXT("ChooseParentClassDescription_Blueprint", "This will add a new Blueprint class to your game project.")
+								)
 						]
 
 						// Full tree checkbox
@@ -350,6 +387,7 @@ void SNewClassDialog::Construct( const FArguments& InArgs )
 							.Padding(0, 0, 12, 0)
 							[
 								SNew(STextBlock)
+								.Visibility(ClassDomain == EClassDomain::Blueprint ? EVisibility::Collapsed : EVisibility::Visible)
 								.TextStyle(FEditorStyle::Get(), "NewClassDialog.SelectedParentClassLabel")
 								.Text(LOCTEXT("ParentClassSourceLabel", "Selected Class Source"))
 							]
@@ -368,6 +406,7 @@ void SNewClassDialog::Construct( const FArguments& InArgs )
 								SNew(SHorizontalBox)
 
 								+SHorizontalBox::Slot()
+								.VAlign(VAlign_Center)
 								.AutoWidth()
 								[
 									SNew(STextBlock)
@@ -375,6 +414,7 @@ void SNewClassDialog::Construct( const FArguments& InArgs )
 								]
 
 								+SHorizontalBox::Slot()
+								.VAlign(VAlign_Center)
 								.AutoWidth()
 								[
 									DocWidget
@@ -434,7 +474,7 @@ void SNewClassDialog::Construct( const FArguments& InArgs )
 						.Padding(0, 0, 0, 5)
 						[
 							SNew(STextBlock)
-							.Text( LOCTEXT("ClassNameDescription", "Enter a name for your new class. Class names may only contain alphanumeric characters, and may not contain a space.") )
+							.Text(LOCTEXT("ClassNameDescription", "Enter a name for your new class. Class names may only contain alphanumeric characters, and may not contain a space.") )
 						]
 
 						+SVerticalBox::Slot()
@@ -442,7 +482,10 @@ void SNewClassDialog::Construct( const FArguments& InArgs )
 						.Padding(0, 0, 0, 2)
 						[
 							SNew(STextBlock)
-							.Text( LOCTEXT("ClassNameDetails", "When you click the \"Create\" button below, a header (.h) file and a source (.cpp) file will be made using this name.") )
+							.Text( ClassDomain == EClassDomain::Native ?
+								LOCTEXT("ClassNameDetails_Native", "When you click the \"Create\" button below, a header (.h) file and a source (.cpp) file will be made using this name.") :
+								LOCTEXT("ClassNameDetails_Blueprint", "When you click the \"Create\" button below, a new Blueprint class will be created.")
+								)
 						]
 
 						// Name Error label
@@ -517,6 +560,7 @@ void SNewClassDialog::Construct( const FArguments& InArgs )
 											.Padding(6.0f, 0.0f, 0.0f, 0.0f)
 											[
 												SAssignNew(AvailableModulesCombo, SComboBox<TSharedPtr<FModuleContextInfo>>)
+												.Visibility(ClassDomain == EClassDomain::Blueprint ? EVisibility::Collapsed : EVisibility::Visible)
 												.ToolTipText( LOCTEXT("ModuleComboToolTip", "Choose the target module for your new class") )
 												.OptionsSource( &AvailableModules )
 												.InitiallySelectedItem( SelectedModuleInfo )
@@ -528,11 +572,13 @@ void SNewClassDialog::Construct( const FArguments& InArgs )
 												]
 											]
 
+											// Native C++ properties
 											+SHorizontalBox::Slot()
 											.AutoWidth()
 											.Padding(6.0f, 0.0f, 0.0f, 0.0f)
 											[
 												SNew(SHorizontalBox)
+												.Visibility(ClassDomain == EClassDomain::Blueprint ? EVisibility::Collapsed : EVisibility::Visible)
 
 												+SHorizontalBox::Slot()
 												.AutoWidth()
@@ -583,7 +629,7 @@ void SNewClassDialog::Construct( const FArguments& InArgs )
 
 									// Path label
 									+SGridPanel::Slot(0, 1)
-									.VAlign(VAlign_Center)
+									.VAlign(ClassDomain == EClassDomain::Blueprint ? VAlign_Top : VAlign_Center)
 									.Padding(0, 0, 12, 0)
 									[
 										SNew(STextBlock)
@@ -596,28 +642,50 @@ void SNewClassDialog::Construct( const FArguments& InArgs )
 									.Padding(0.0f, 3.0f)
 									.VAlign(VAlign_Center)
 									[
-										SNew(SBox)
-										.HeightOverride(EditableTextHeight)
-										.AddMetaData<FTutorialMetaData>(TEXT("Path"))
-										[
-											SNew(SHorizontalBox)
+										SNew(SVerticalBox)
 
-											+SHorizontalBox::Slot()
-											.FillWidth(1.0f)
+										// Blueprint Class asset path
+										+ SVerticalBox::Slot()
+										.Padding(0)
+										[
+											SNew(SBox)
+											// Height override to force the visibility of a scrollbar (our parent is autoheight)
+											.HeightOverride(200)
+											.Visibility(ClassDomain == EClassDomain::Blueprint ? EVisibility::Visible : EVisibility::Collapsed)
 											[
-												SNew(SEditableTextBox)
-												.Text(this, &SNewClassDialog::OnGetClassPathText)
-												.OnTextChanged(this, &SNewClassDialog::OnClassPathTextChanged)
+												ContentBrowser.CreatePathPicker(BlueprintPathConfig)
 											]
+										]
 
-											+SHorizontalBox::Slot()
-											.AutoWidth()
-											.Padding(6.0f, 1.0f, 0.0f, 0.0f)
+										// Native C++ path
+										+ SVerticalBox::Slot()
+										.Padding(0)
+										.AutoHeight()
+										[
+											SNew(SBox)
+											.Visibility(ClassDomain == EClassDomain::Blueprint ? EVisibility::Collapsed : EVisibility::Visible)
+											.HeightOverride(EditableTextHeight)
+											.AddMetaData<FTutorialMetaData>(TEXT("Path"))
 											[
-												SNew(SButton)
-												.VAlign(VAlign_Center)
-												.OnClicked(this, &SNewClassDialog::HandleChooseFolderButtonClicked)
-												.Text( LOCTEXT( "BrowseButtonText", "Choose Folder" ) )
+												SNew(SHorizontalBox)
+
+												+SHorizontalBox::Slot()
+												.FillWidth(1.0f)
+												[
+													SNew(SEditableTextBox)
+													.Text(this, &SNewClassDialog::OnGetClassPathText)
+													.OnTextChanged(this, &SNewClassDialog::OnClassPathTextChanged)
+												]
+
+												+SHorizontalBox::Slot()
+												.AutoWidth()
+												.Padding(6.0f, 1.0f, 0.0f, 0.0f)
+												[
+													SNew(SButton)
+													.VAlign(VAlign_Center)
+													.OnClicked(this, &SNewClassDialog::HandleChooseFolderButtonClicked)
+													.Text( LOCTEXT( "BrowseButtonText", "Choose Folder" ) )
+												]
 											]
 										]
 									]
@@ -628,6 +696,7 @@ void SNewClassDialog::Construct( const FArguments& InArgs )
 									.Padding(0, 0, 12, 0)
 									[
 										SNew(STextBlock)
+										.Visibility(ClassDomain == EClassDomain::Blueprint ? EVisibility::Collapsed : EVisibility::Visible)
 										.TextStyle( FEditorStyle::Get(), "NewClassDialog.SelectedParentClassLabel" )
 										.Text( LOCTEXT( "HeaderFileLabel", "Header File" ) )
 									]
@@ -638,6 +707,7 @@ void SNewClassDialog::Construct( const FArguments& InArgs )
 									.VAlign(VAlign_Center)
 									[
 										SNew(SBox)
+										.Visibility(ClassDomain == EClassDomain::Blueprint ? EVisibility::Collapsed : EVisibility::Visible)
 										.VAlign(VAlign_Center)
 										.HeightOverride(EditableTextHeight)
 										[
@@ -652,6 +722,7 @@ void SNewClassDialog::Construct( const FArguments& InArgs )
 									.Padding(0, 0, 12, 0)
 									[
 										SNew(STextBlock)
+										.Visibility(ClassDomain == EClassDomain::Blueprint ? EVisibility::Collapsed : EVisibility::Visible)
 										.TextStyle( FEditorStyle::Get(), "NewClassDialog.SelectedParentClassLabel" )
 										.Text( LOCTEXT( "SourceFileLabel", "Source File" ) )
 									]
@@ -662,6 +733,7 @@ void SNewClassDialog::Construct( const FArguments& InArgs )
 									.VAlign(VAlign_Center)
 									[
 										SNew(SBox)
+										.Visibility(ClassDomain == EClassDomain::Blueprint ? EVisibility::Collapsed : EVisibility::Visible)
 										.VAlign(VAlign_Center)
 										.HeightOverride(EditableTextHeight)
 										[
@@ -745,8 +817,7 @@ TSharedRef<ITableRow> SNewClassDialog::MakeParentClassListViewWidget(TSharedPtr<
 					.Padding(0, 0, 4, 0)
 					[
 						SNew(SImage)
-						.ColorAndOpacity(FSlateColor::UseForeground())
-						.Image( ClassBrush )
+						.Image(ClassBrush)
 					]
 
 					+SHorizontalBox::Slot()
@@ -790,6 +861,11 @@ FString GetClassHeaderPath(const UClass* Class)
 
 EVisibility SNewClassDialog::GetSourceHyperlinkVisibility() const
 {
+	if (ClassDomain == EClassDomain::Blueprint)
+	{
+		return EVisibility::Collapsed;
+	}
+
 	return (GetClassHeaderPath(ParentClassInfo.BaseClass).Len() > 0 ? EVisibility::Visible : EVisibility::Hidden);
 }
 
@@ -843,14 +919,14 @@ void SNewClassDialog::OnClassSelected(TSharedPtr<FParentClassItem> Item, ESelect
 	}
 	else
 	{
-		ParentClassInfo = GameProjectUtils::FNewClassInfo();
+		ParentClassInfo = FNewClassInfo();
 	}
 }
 
 void SNewClassDialog::OnAdvancedClassSelected(UClass* Class)
 {
 	ParentClassListView->ClearSelection();
-	ParentClassInfo = GameProjectUtils::FNewClassInfo(Class);
+	ParentClassInfo = FNewClassInfo(Class);
 }
 
 ECheckBoxState SNewClassDialog::IsFullClassTreeChecked() const
@@ -971,6 +1047,12 @@ void SNewClassDialog::OnClassPathTextChanged(const FText& NewText)
 	UpdateInputValidity();
 }
 
+void SNewClassDialog::OnBlueprintPathSelected(const FString& NewPath)
+{
+	NewClassPath = NewPath;
+	UpdateInputValidity();
+}
+
 FText SNewClassDialog::OnGetClassHeaderFileText() const
 {
 	return FText::FromString(CalculatedClassHeaderName);
@@ -995,77 +1077,125 @@ void SNewClassDialog::FinishClicked()
 {
 	check(CanFinish());
 
-	FString HeaderFilePath;
-	FString CppFilePath;
-
-	FText FailReason;
-	const TSet<FString>& DisallowedHeaderNames = FSourceCodeNavigation::GetSourceFileDatabase().GetDisallowedHeaderNames();
-	if (GameProjectUtils::AddCodeToProject(NewClassName, NewClassPath, *SelectedModuleInfo, ParentClassInfo, DisallowedHeaderNames, HeaderFilePath, CppFilePath, FailReason))
+	if (ClassDomain == EClassDomain::Blueprint)
 	{
-		OnCodeAddedToProject.ExecuteIfBound( NewClassName, NewClassPath, SelectedModuleInfo->ModuleName );
+		FString PackagePath = NewClassPath / NewClassName;
 
-		// Prevent periodic validity checks. This is to prevent a brief error message about the class already existing while you are exiting.
-		bPreventPeriodicValidityChecksUntilNextChange = true;
-
-		// Display a nag if we didn't automatically hot-reload for the newly added class
-		const bool bWasHotReloaded = GEditor->AccessEditorUserSettings().bAutomaticallyHotReloadNewClasses;
-		if( bWasHotReloaded )
+		if (!ParentClassInfo.BaseClass)
 		{
-			FNotificationInfo Notification( FText::Format( LOCTEXT("AddedClassSuccessNotification", "Added new class {0}"), FText::FromString(NewClassName) ) );
-			FSlateNotificationManager::Get().AddNotification( Notification );
+			// @todo show fail reason in error label
+			FMessageDialog::Open(EAppMsgType::Ok, LOCTEXT("AddCodeFailed_Blueprint_NoBase", "No parent class has been specified. Failed to generate new Blueprint class."));
 		}
-
-		if ( HeaderFilePath.IsEmpty() || CppFilePath.IsEmpty() || !FSlateApplication::Get().SupportsSourceAccess() )
+		else if (FindObject<UBlueprint>(ANY_PACKAGE, *PackagePath))
 		{
-			if( !bWasHotReloaded )
-			{
-				// Code successfully added, notify the user. We are either running on a platform that does not support source access or a file was not given so don't ask about editing the file
-				const FText Message = FText::Format( 
-					LOCTEXT("AddCodeSuccessWithHotReload", "Successfully added class {0}, however you must recompile {1} before it will appear in the Content Browser.")
-					, FText::FromString(NewClassName), FText::FromString(SelectedModuleInfo->ModuleName) );
-				FMessageDialog::Open(EAppMsgType::Ok, Message);
-			}
-			else
+			// @todo show fail reason in error label
+			FMessageDialog::Open(EAppMsgType::Ok, LOCTEXT("AddCodeFailed_Blueprint_AlreadyExists", "The chosen Blueprint class already exists, please try again with a different name."));
+		}
+		else if (!NewClassPath.IsEmpty() && !NewClassName.IsEmpty())
+		{
+			UPackage* Package = CreatePackage(NULL, *PackagePath);
+			if (Package)
 			{
-				// Code was added and hot reloaded into the editor, but the user doesn't have a code IDE installed so we can't open the file to edit it now
+				// Create and init a new Blueprint
+				UBlueprint* NewBP = FKismetEditorUtilities::CreateBlueprint(const_cast<UClass*>(ParentClassInfo.BaseClass), Package, FName(*NewClassName), BPTYPE_Normal, UBlueprint::StaticClass(), UBlueprintGeneratedClass::StaticClass());
+				if (NewBP)
+				{
+					// Notify the asset registry
+					FAssetRegistryModule::AssetCreated(NewBP);
+
+					// Mark the package dirty...
+					Package->MarkPackageDirty();
+
+					OnAddedToProject.ExecuteIfBound( NewClassName, PackagePath, FString() );
+
+					// Successfully created the code and potentially opened the IDE. Close the dialog.
+					CloseContainingWindow();
+
+					return;
+				}
 			}
 		}
-		else
+
+		// @todo show fail reason in error label
+		// Failed to add blueprint
+		const FText Message = FText::Format( LOCTEXT("AddCodeFailed_Blueprint", "Failed to create package for class {0}. Please try again with a different name."), FText::FromString(NewClassName) );
+		FMessageDialog::Open(EAppMsgType::Ok, Message);
+	}
+	else
+	{
+
+		FString HeaderFilePath;
+		FString CppFilePath;
+
+		FText FailReason;
+		const TSet<FString>& DisallowedHeaderNames = FSourceCodeNavigation::GetSourceFileDatabase().GetDisallowedHeaderNames();
+		if (GameProjectUtils::AddCodeToProject(NewClassName, NewClassPath, *SelectedModuleInfo, ParentClassInfo, DisallowedHeaderNames, HeaderFilePath, CppFilePath, FailReason))
 		{
-			bool bEditSourceFilesNow = false;
+			OnAddedToProject.ExecuteIfBound( NewClassName, NewClassPath, SelectedModuleInfo->ModuleName );
+
+			// Prevent periodic validity checks. This is to prevent a brief error message about the class already existing while you are exiting.
+			bPreventPeriodicValidityChecksUntilNextChange = true;
+
+			// Display a nag if we didn't automatically hot-reload for the newly added class
+			const bool bWasHotReloaded = GEditor->AccessEditorUserSettings().bAutomaticallyHotReloadNewClasses;
 			if( bWasHotReloaded )
 			{
-				// Code was hot reloaded, so always edit the new classes now
-				bEditSourceFilesNow = true;
+				FNotificationInfo Notification( FText::Format( LOCTEXT("AddedClassSuccessNotification", "Added new class {0}"), FText::FromString(NewClassName) ) );
+				FSlateNotificationManager::Get().AddNotification( Notification );
 			}
-			else
+
+			if ( HeaderFilePath.IsEmpty() || CppFilePath.IsEmpty() || !FSlateApplication::Get().SupportsSourceAccess() )
 			{
-				// Code successfully added, notify the user and ask about opening the IDE now
-				const FText Message = FText::Format( 
-					LOCTEXT("AddCodeSuccessWithHotReloadAndSync", "Successfully added class {0}, however you must recompile {1} before it will appear in the Content Browser.\n\nWould you like to edit the code now?")
-					, FText::FromString(NewClassName), FText::FromString(SelectedModuleInfo->ModuleName) );
-				bEditSourceFilesNow = ( FMessageDialog::Open( EAppMsgType::YesNo, Message ) == EAppReturnType::Yes );
+				if( !bWasHotReloaded )
+				{
+					// Code successfully added, notify the user. We are either running on a platform that does not support source access or a file was not given so don't ask about editing the file
+					const FText Message = FText::Format( 
+						LOCTEXT("AddCodeSuccessWithHotReload", "Successfully added class {0}, however you must recompile {1} before it will appear in the Content Browser.")
+						, FText::FromString(NewClassName), FText::FromString(SelectedModuleInfo->ModuleName) );
+					FMessageDialog::Open(EAppMsgType::Ok, Message);
+				}
+				else
+				{
+					// Code was added and hot reloaded into the editor, but the user doesn't have a code IDE installed so we can't open the file to edit it now
+				}
 			}
-
-			if( bEditSourceFilesNow )
+			else
 			{
-				TArray<FString> SourceFiles;
-				SourceFiles.Add(IFileManager::Get().ConvertToAbsolutePathForExternalAppForRead(*HeaderFilePath));
-				SourceFiles.Add(IFileManager::Get().ConvertToAbsolutePathForExternalAppForRead(*CppFilePath));
+				bool bEditSourceFilesNow = false;
+				if( bWasHotReloaded )
+				{
+					// Code was hot reloaded, so always edit the new classes now
+					bEditSourceFilesNow = true;
+				}
+				else
+				{
+					// Code successfully added, notify the user and ask about opening the IDE now
+					const FText Message = FText::Format( 
+						LOCTEXT("AddCodeSuccessWithHotReloadAndSync", "Successfully added class {0}, however you must recompile {1} before it will appear in the Content Browser.\n\nWould you like to edit the code now?")
+						, FText::FromString(NewClassName), FText::FromString(SelectedModuleInfo->ModuleName) );
+					bEditSourceFilesNow = ( FMessageDialog::Open( EAppMsgType::YesNo, Message ) == EAppReturnType::Yes );
+				}
 
-				FSourceCodeNavigation::OpenSourceFiles(SourceFiles);
+				if( bEditSourceFilesNow )
+				{
+					TArray<FString> SourceFiles;
+					SourceFiles.Add(IFileManager::Get().ConvertToAbsolutePathForExternalAppForRead(*HeaderFilePath));
+					SourceFiles.Add(IFileManager::Get().ConvertToAbsolutePathForExternalAppForRead(*CppFilePath));
+
+					FSourceCodeNavigation::OpenSourceFiles(SourceFiles);
+				}
 			}
-		}
 
-		// Successfully created the code and potentially opened the IDE. Close the dialog.
-		CloseContainingWindow();
-	}
-	else
-	{
-		// @todo show fail reason in error label
-		// Failed to add code
-		const FText Message = FText::Format( LOCTEXT("AddCodeFailed", "Failed to add class {0}. {1}"), FText::FromString(NewClassName), FailReason );
-		FMessageDialog::Open(EAppMsgType::Ok, Message);
+			// Successfully created the code and potentially opened the IDE. Close the dialog.
+			CloseContainingWindow();
+		}
+		else
+		{
+			// @todo show fail reason in error label
+			// Failed to add code
+			const FText Message = FText::Format( LOCTEXT("AddCodeFailed_Native", "Failed to add class {0}. {1}"), FText::FromString(NewClassName), FailReason );
+			FMessageDialog::Open(EAppMsgType::Ok, Message);
+		}
 	}
 }
 
@@ -1233,46 +1363,63 @@ void SNewClassDialog::UpdateInputValidity()
 {
 	bLastInputValidityCheckSuccessful = true;
 
-	// Validate the path first since this has the side effect of updating the UI
-	bLastInputValidityCheckSuccessful = GameProjectUtils::CalculateSourcePaths(NewClassPath, *SelectedModuleInfo, CalculatedClassHeaderName, CalculatedClassSourceName, &LastInputValidityErrorText);
-	CalculatedClassHeaderName /= ParentClassInfo.GetHeaderFilename(NewClassName);
-	CalculatedClassSourceName /= ParentClassInfo.GetSourceFilename(NewClassName);
-
-	// If the source paths check as succeeded, check to see if we're using a Public/Private class
-	if(bLastInputValidityCheckSuccessful)
+	if (ClassDomain == EClassDomain::Blueprint)
 	{
-		GameProjectUtils::GetClassLocation(NewClassPath, *SelectedModuleInfo, ClassLocation);
-
-		// We only care about the Public and Private folders
-		if(ClassLocation != GameProjectUtils::EClassLocation::Public && ClassLocation != GameProjectUtils::EClassLocation::Private)
+		bLastInputValidityCheckSuccessful = GameProjectUtils::IsValidClassNameForCreation(NewClassName, LastInputValidityErrorText);
+		ClassLocation = GameProjectUtils::EClassLocation::UserDefined;
+		if (bLastInputValidityCheckSuccessful)
 		{
-			ClassLocation = GameProjectUtils::EClassLocation::UserDefined;
+			IAssetRegistry& AssetRegistry = FModuleManager::LoadModuleChecked<FAssetRegistryModule>(AssetRegistryConstants::ModuleName).Get();
+			if (AssetRegistry.GetAssetByObjectPath(*(NewClassPath / NewClassName)).IsValid())
+			{
+				bLastInputValidityCheckSuccessful = false;
+				LastInputValidityErrorText = FText::Format(LOCTEXT("AssetAlreadyExists", "An asset called {0} already exists in {1}."), FText::FromString(NewClassName), FText::FromString(NewClassPath));
+			}
 		}
 	}
 	else
 	{
-		ClassLocation = GameProjectUtils::EClassLocation::UserDefined;
-	}
+		// Validate the path first since this has the side effect of updating the UI
+		bLastInputValidityCheckSuccessful = GameProjectUtils::CalculateSourcePaths(NewClassPath, *SelectedModuleInfo, CalculatedClassHeaderName, CalculatedClassSourceName, &LastInputValidityErrorText);
+		CalculatedClassHeaderName /= ParentClassInfo.GetHeaderFilename(NewClassName);
+		CalculatedClassSourceName /= ParentClassInfo.GetSourceFilename(NewClassName);
 
-	// Validate the class name only if the path is valid
-	if ( bLastInputValidityCheckSuccessful )
-	{
-		const TSet<FString>& DisallowedHeaderNames = FSourceCodeNavigation::GetSourceFileDatabase().GetDisallowedHeaderNames();
-		bLastInputValidityCheckSuccessful = GameProjectUtils::IsValidClassNameForCreation(NewClassName, *SelectedModuleInfo, DisallowedHeaderNames, LastInputValidityErrorText);
-	}
+		// If the source paths check as succeeded, check to see if we're using a Public/Private class
+		if(bLastInputValidityCheckSuccessful)
+		{
+			GameProjectUtils::GetClassLocation(NewClassPath, *SelectedModuleInfo, ClassLocation);
 
-	// Validate that the class is valid for the currently selected module
-	// As a project can have multiple modules, this lets us update the class validity as the user changes the target module
-	if ( bLastInputValidityCheckSuccessful && ParentClassInfo.BaseClass )
-	{
-		bLastInputValidityCheckSuccessful = GameProjectUtils::IsValidBaseClassForCreation(ParentClassInfo.BaseClass, *SelectedModuleInfo);
-		if ( !bLastInputValidityCheckSuccessful )
+			// We only care about the Public and Private folders
+			if(ClassLocation != GameProjectUtils::EClassLocation::Public && ClassLocation != GameProjectUtils::EClassLocation::Private)
+			{
+				ClassLocation = GameProjectUtils::EClassLocation::UserDefined;
+			}
+		}
+		else
+		{
+			ClassLocation = GameProjectUtils::EClassLocation::UserDefined;
+		}
+
+		// Validate the class name only if the path is valid
+		if ( bLastInputValidityCheckSuccessful )
 		{
-			LastInputValidityErrorText = FText::Format(
-				LOCTEXT("NewClassError_InvalidBaseClassForModule", "{0} cannot be used as a base class in the {1} module. Please make sure that {0} is API exported."),
-				FText::FromString(ParentClassInfo.BaseClass->GetName()),
-				FText::FromString(SelectedModuleInfo->ModuleName)
-				);
+			const TSet<FString>& DisallowedHeaderNames = FSourceCodeNavigation::GetSourceFileDatabase().GetDisallowedHeaderNames();
+			bLastInputValidityCheckSuccessful = GameProjectUtils::IsValidClassNameForCreation(NewClassName, *SelectedModuleInfo, DisallowedHeaderNames, LastInputValidityErrorText);
+		}
+
+		// Validate that the class is valid for the currently selected module
+		// As a project can have multiple modules, this lets us update the class validity as the user changes the target module
+		if ( bLastInputValidityCheckSuccessful && ParentClassInfo.BaseClass )
+		{
+			bLastInputValidityCheckSuccessful = GameProjectUtils::IsValidBaseClassForCreation(ParentClassInfo.BaseClass, *SelectedModuleInfo);
+			if ( !bLastInputValidityCheckSuccessful )
+			{
+				LastInputValidityErrorText = FText::Format(
+					LOCTEXT("NewClassError_InvalidBaseClassForModule", "{0} cannot be used as a base class in the {1} module. Please make sure that {0} is API exported."),
+					FText::FromString(ParentClassInfo.BaseClass->GetName()),
+					FText::FromString(SelectedModuleInfo->ModuleName)
+					);
+			}
 		}
 	}
 
@@ -1282,38 +1429,26 @@ void SNewClassDialog::UpdateInputValidity()
 	bPreventPeriodicValidityChecksUntilNextChange = false;
 }
 
-const GameProjectUtils::FNewClassInfo& SNewClassDialog::GetSelectedParentClassInfo() const
+const FNewClassInfo& SNewClassDialog::GetSelectedParentClassInfo() const
 {
 	return ParentClassInfo;
 }
 
-void SNewClassDialog::SetupParentClassItems()
+void SNewClassDialog::SetupParentClassItems(const TArray<FNewClassInfo>& UserSpecifiedFeaturedClasses)
 {
-	TArray<GameProjectUtils::FNewClassInfo> FeaturedClasses;
-	
-	// Add an empty class
-	FeaturedClasses.Add(GameProjectUtils::FNewClassInfo(GameProjectUtils::FNewClassInfo::EClassType::EmptyCpp));
-
-	// @todo make this ini configurable
-	FeaturedClasses.Add(GameProjectUtils::FNewClassInfo(ACharacter::StaticClass()));
-	FeaturedClasses.Add(GameProjectUtils::FNewClassInfo(APawn::StaticClass()));
-	FeaturedClasses.Add(GameProjectUtils::FNewClassInfo(AActor::StaticClass()));
-	FeaturedClasses.Add(GameProjectUtils::FNewClassInfo(APlayerCameraManager::StaticClass()));
-	FeaturedClasses.Add(GameProjectUtils::FNewClassInfo(APlayerController::StaticClass()));
-	FeaturedClasses.Add(GameProjectUtils::FNewClassInfo(AGameMode::StaticClass()));
-	FeaturedClasses.Add(GameProjectUtils::FNewClassInfo(AWorldSettings::StaticClass()));
-	FeaturedClasses.Add(GameProjectUtils::FNewClassInfo(AHUD::StaticClass()));
-	FeaturedClasses.Add(GameProjectUtils::FNewClassInfo(APlayerState::StaticClass()));
-	FeaturedClasses.Add(GameProjectUtils::FNewClassInfo(AGameState::StaticClass()));
-	FeaturedClasses.Add(GameProjectUtils::FNewClassInfo(UBlueprintFunctionLibrary::StaticClass()));
-
-	// Add the extra non-UObject classes
-	FeaturedClasses.Add(GameProjectUtils::FNewClassInfo(GameProjectUtils::FNewClassInfo::EClassType::SlateWidget));
-	FeaturedClasses.Add(GameProjectUtils::FNewClassInfo(GameProjectUtils::FNewClassInfo::EClassType::SlateWidgetStyle));
-
-	for ( auto ClassIt = FeaturedClasses.CreateConstIterator(); ClassIt; ++ClassIt )
+	TArray<FNewClassInfo> DefaultFeaturedClasses;
+	const TArray<FNewClassInfo>* ArrayToUse = &UserSpecifiedFeaturedClasses;
+
+	// Setup the featured classes list
+	if (ArrayToUse->Num() == 0)
+	{
+		DefaultFeaturedClasses = ClassDomain == EClassDomain::Native ? FFeaturedClasses::AllNativeClasses() : FFeaturedClasses::ActorClasses();
+		ArrayToUse = &DefaultFeaturedClasses;
+	}
+
+	for (const auto& Featured : *ArrayToUse)
 	{
-		ParentClassItemsSource.Add( MakeShareable( new FParentClassItem(*ClassIt) ) );
+		ParentClassItemsSource.Add( MakeShareable( new FParentClassItem(Featured) ) );
 	}
 }
 
diff --git a/Engine/Source/Editor/GameProjectGeneration/Private/SNewClassDialog.h b/Engine/Source/Editor/GameProjectGeneration/Private/SNewClassDialog.h
--- a/Engine/Source/Editor/GameProjectGeneration/Private/SNewClassDialog.h
+++ b/Engine/Source/Editor/GameProjectGeneration/Private/SNewClassDialog.h
@@ -3,6 +3,9 @@
 #pragma once
 
 struct FParentClassItem;
+class IClassViewerFilter;
+
+enum class EClassDomain : uint8 { Blueprint, Native };
 
 /**
  * A dialog to choose a new class parent and name
@@ -11,9 +14,18 @@ class SNewClassDialog : public SCompoundWidget
 {
 public:
 	SLATE_BEGIN_ARGS( SNewClassDialog )
-		: _Class(NULL)
+		: _ClassDomain(EClassDomain::Native), _Class(NULL)
 	{}
 
+	/** The domain of the class we are to create (native or blueprint) */
+	SLATE_ARGUMENT(EClassDomain, ClassDomain)
+
+	/** An array of classes to feature on the class picker page */
+	SLATE_ARGUMENT(TArray<FNewClassInfo>, FeaturedClasses)
+
+	/** Filter specifying allowable class types, if a parent class is to be chosen by the user */
+	SLATE_ARGUMENT(TSharedPtr<IClassViewerFilter>, ClassViewerFilter)
+
 	/** The class we want to build our new class from. If this is not specified then the wizard will display classes to the user. */
 	SLATE_ARGUMENT(const UClass*, Class)
 
@@ -28,7 +40,7 @@ public:
 	SLATE_ARGUMENT(FString, DefaultClassName)
 
 	/** Event called when code is successfully added to the project */
-	SLATE_EVENT( FOnCodeAddedToProject, OnCodeAddedToProject )
+	SLATE_EVENT( FOnAddedToProject, OnAddedToProject )
 	SLATE_END_ARGS()
 
 	/** Constructs this widget with InArgs */
@@ -110,6 +122,9 @@ private:
 	/** Handler for when the text in the class path edit box has changed */
 	void OnClassPathTextChanged(const FText& NewText);
 
+	/** Called when the user chooses a path for a blueprint */
+	void OnBlueprintPathSelected(const FString& NewPath);
+
 	/** Returns the text for the calculated header file name */
 	FText OnGetClassHeaderFileText() const;
 
@@ -158,10 +173,10 @@ private:
 	void UpdateInputValidity();
 
 	/** Gets the currently selected parent class */
-	const GameProjectUtils::FNewClassInfo& GetSelectedParentClassInfo() const;
+	const FNewClassInfo& GetSelectedParentClassInfo() const;
 
 	/** Adds parent classes to the ParentClassListView source */
-	void SetupParentClassItems();
+	void SetupParentClassItems(const TArray<FNewClassInfo>& UserSpecifiedFeaturedClasses);
 
 	/** Closes the window that contains this widget */
 	void CloseContainingWindow();
@@ -207,7 +222,7 @@ private:
 	FString LastAutoGeneratedClassName;
 
 	/** The selected parent class */
-	GameProjectUtils::FNewClassInfo ParentClassInfo;
+	FNewClassInfo ParentClassInfo;
 
 	/** If true, the full class tree will be shown in the parent class selection */
 	bool bShowFullClassTree;
@@ -230,6 +245,9 @@ private:
 	/** Whether the class should be created as a Public or Private class */
 	GameProjectUtils::EClassLocation ClassLocation;
 
+	/** The domain of the new class we are creating (native or blueprint) */
+	EClassDomain ClassDomain;
+
 	/** Information about the currently available modules for this project */
 	TArray<TSharedPtr<FModuleContextInfo>> AvailableModules;
 
@@ -237,5 +255,5 @@ private:
 	TSharedPtr<FModuleContextInfo> SelectedModuleInfo;
 
 	/** Event called when code is succesfully added to the project */
-	FOnCodeAddedToProject OnCodeAddedToProject;
+	FOnAddedToProject OnAddedToProject;
 };
diff --git a/Engine/Source/Editor/GameProjectGeneration/Public/AddToProjectConfig.h b/Engine/Source/Editor/GameProjectGeneration/Public/AddToProjectConfig.h
new file mode 100644
--- /dev/null
+++ b/Engine/Source/Editor/GameProjectGeneration/Public/AddToProjectConfig.h
@@ -0,0 +1,181 @@
+// Copyright 1998-2015 Epic Games, Inc. All Rights Reserved.
+
+#pragma once
+
+
+/**
+ * Delegate called when code is added to the project.  Passes in the created class name and class path
+ * 
+ * @param ClassName		The created class name
+ * @param ClassPath		The created class path
+ * @param ModuleName	The name of the module that the class was added to
+ */
+DECLARE_DELEGATE_ThreeParams( FOnAddedToProject, const FString& /*ClassName*/, const FString& /*ClassPath*/, const FString& /*ModuleName*/);
+
+class IClassViewerFilter;
+
+/** Information used when creating a new class via AddCodeToProject */
+struct FNewClassInfo
+{
+	/** The type of class we want to create */
+	enum class EClassType : uint8
+	{
+		/** The new class is using a UObject as a base, consult BaseClass for the type */
+		UObject,
+
+		/** The new class should be an empty standard C++ class */
+		EmptyCpp,
+
+		/** The new class should be a Slate widget, deriving from SCompoundWidget */
+		SlateWidget,
+
+		/** The new class should be a Slate widget style, deriving from FSlateWidgetStyle, along with its associated UObject wrapper class */
+		SlateWidgetStyle,
+	};
+
+	/** Default constructor; must produce an object which fails the IsSet check */
+	FNewClassInfo()
+		: ClassType(EClassType::UObject)
+		, BaseClass(nullptr)
+	{
+	}
+
+	/** Convenience constructor so you can construct from a EClassType */
+	explicit FNewClassInfo(const EClassType InClassType)
+		: ClassType(InClassType)
+		, BaseClass(nullptr)
+	{
+	}
+
+	/** Convenience constructor so you can construct from a UClass */
+	explicit FNewClassInfo(const UClass* InBaseClass)
+		: ClassType(EClassType::UObject)
+		, BaseClass(InBaseClass)
+	{
+	}
+
+	/** Check to see if this struct is set to something that could be used to create a new class */
+	bool IsSet() const
+	{
+		return ClassType != EClassType::UObject || BaseClass;
+	}
+
+	/** Get the "friendly" class name to use in the UI */
+	FString GetClassName() const;
+
+	/** Get the class description to use in the UI */
+	FString GetClassDescription() const;
+
+	/** Get the class icon to use in the UI */
+	const FSlateBrush* GetClassIcon() const;
+
+	/** Get the C++ prefix used for this class type */
+	FString GetClassPrefixCPP() const;
+
+	/** Get the C++ class name; this may or may not be prefixed, but will always produce a valid C++ name via GetClassPrefix() + GetClassName() */
+	FString GetClassNameCPP() const;
+
+	/** Some classes may apply a particular suffix; this function returns the class name with those suffixes removed */
+	FString GetCleanClassName(const FString& ClassName) const;
+
+	/** Some classes may apply a particular suffix; this function returns the class name that will ultimately be used should that happen */
+	FString GetFinalClassName(const FString& ClassName) const;
+
+	/** Get the path needed to include this class into another file */
+	bool GetIncludePath(FString& OutIncludePath) const;
+
+	/** Given a class name, generate the header file (.h) that should be used for this class */
+	FString GetHeaderFilename(const FString& ClassName) const;
+
+	/** Given a class name, generate the source file (.cpp) that should be used for this class */
+	FString GetSourceFilename(const FString& ClassName) const;
+
+	/** Get the generation template filename to used based on the current class type */
+	FString GetHeaderTemplateFilename() const;
+
+	/** Get the generation template filename to used based on the current class type */
+	FString GetSourceTemplateFilename() const;
+
+	/** The type of class we want to create */
+	EClassType ClassType;
+
+	/** Base class information; if the ClassType is UObject */
+	const UClass* BaseClass;
+};
+
+/** Helper that creates lists of featured classes. Include FeaturedClasses.inl for definitions. */
+/** @todo make this ini configurable */
+struct FFeaturedClasses
+{
+	/** Get a list of all featured native class types */
+	static TArray<FNewClassInfo> AllNativeClasses();
+	
+	/** Get a list of all featured Actor class types */
+	static TArray<FNewClassInfo> ActorClasses();
+
+	/** Get a list of all featured class types */
+	static TArray<FNewClassInfo> ComponentClasses();
+};
+
+/** Add to project dialog configuration structure */
+struct FAddToProjectConfig
+{
+	FAddToProjectConfig() : _ParentClass(nullptr), _bModal(false)	{}
+
+public:
+
+	/** Force the add to project dialog to use the specified parent class. Skips the first step (choose a parent class) as a result. */
+	FAddToProjectConfig& ParentClass(const UClass* InClass)							{ _ParentClass = InClass; return *this; }
+	/** Limits the allowable parent classes by the specified filter. */
+	FAddToProjectConfig& AllowableParents(const TSharedPtr<IClassViewerFilter>& In)	{ _ParentClass = nullptr; _AllowableParents = In; return *this; }
+
+	/** The initial path we should use as the destination for the new file, or an empty string to choose a suitable default based upon the module path. */
+	FAddToProjectConfig& InitialPath(FString InInitialPath)							{ _InitialPath = MoveTemp(InInitialPath); return *this; }
+	/** Optional argument that specifies the default name for the new class being added.  The user will be able to type their own name if they don't like this name.  If empty, defaults to the name of the inherited class. */
+	FAddToProjectConfig& DefaultClassName(FString InDefaultClassName)				{ _DefaultClassName = MoveTemp(InDefaultClassName); return *this; }
+	/** Optional argument that specifies the prefix for the new class name.  The user will be able to type their own name if they don't like this name.  Defaults to "My" if not specified or empty. */
+	FAddToProjectConfig& DefaultClassPrefix(FString InDefaultClassPrefix)			{ _DefaultClassPrefix = MoveTemp(InDefaultClassPrefix); return *this; }
+
+	/** The title text to display on the window */
+	FAddToProjectConfig& WindowTitle(FText InText)									{ _WindowTitle = MoveTemp(InText); return *this; }	
+	/** The parent window the dialog should use, or null to choose a suitable default parent window (the main frame, if available). */
+	FAddToProjectConfig& ParentWindow(const TSharedPtr<SWindow>& InParentWindow)	{ _ParentWindow = InParentWindow; return *this; }
+	/** Make the window modal to force the user to make a decision before continuing. */
+	FAddToProjectConfig& Modal(bool bModal = true)									{ _bModal = bModal; return *this; }
+
+	/** Callback for when the object is successfully added to the project */
+	FAddToProjectConfig& OnAddedToProject(const FOnAddedToProject& InDelegate)		{ _OnAddedToProject = InDelegate; return *this; }
+
+	/** Set the add to project dialog to show component types on the initial 'featured' classes list */
+	FAddToProjectConfig& FeatureAllNativeClasses()									{ _FeaturedClasses = FFeaturedClasses::AllNativeClasses(); return *this; }
+	/** Set the add to project dialog to show component types on the initial 'featured' classes list */
+	FAddToProjectConfig& FeatureActorClasses()										{ _FeaturedClasses = FFeaturedClasses::ActorClasses(); return *this; }
+	/** Set the add to project dialog to show component types on the initial 'featured' classes list */
+	FAddToProjectConfig& FeatureComponentClasses()									{ _FeaturedClasses = FFeaturedClasses::ComponentClasses(); return *this; }
+
+public:
+
+	/** Forced parent class to use */
+	const UClass* _ParentClass;
+	/** Filter for allowable parent classes, when ParentClass is nullptr */
+	TSharedPtr<IClassViewerFilter> _AllowableParents;
+	/** Array of featured classes */
+	TArray<FNewClassInfo> _FeaturedClasses;
+
+	/** Initial file path for the (blueprint) class */
+	FString _InitialPath;
+	/** Default name prefix for the (blueprint) class */
+	FString _DefaultClassPrefix;
+	/** Default name for the (blueprint) class, excluding class prefix */
+	FString _DefaultClassName;
+
+	/** The title to display on the window */
+	FText _WindowTitle;
+	/** Parent window to use */
+	TSharedPtr<SWindow> _ParentWindow;
+	/** True to force a modal dialog, false otherwise */
+	bool _bModal;
+
+	/** Delegate to invoke when the (blueprint) class has been added to the project */
+	FOnAddedToProject _OnAddedToProject;
+};
\ No newline at end of file
diff --git a/Engine/Source/Editor/GameProjectGeneration/Public/FeaturedClasses.inl b/Engine/Source/Editor/GameProjectGeneration/Public/FeaturedClasses.inl
new file mode 100644
--- /dev/null
+++ b/Engine/Source/Editor/GameProjectGeneration/Public/FeaturedClasses.inl
@@ -0,0 +1,68 @@
+// Copyright 1998-2015 Epic Games, Inc. All Rights Reserved.
+
+/** Inline file to avoid introducing many UObject headers into the global namespace */
+
+#include "Kismet/BlueprintFunctionLibrary.h"
+#include "GameFramework/Character.h"
+#include "GameFramework/Pawn.h"
+#include "GameFramework/Actor.h"
+#include "GameFramework/PlayerController.h"
+#include "GameFramework/GameMode.h"
+#include "GameFramework/WorldSettings.h"
+#include "GameFramework/HUD.h"
+#include "GameFramework/PlayerState.h"
+#include "GameFramework/GameState.h"
+#include "Components/ActorComponent.h"
+#include "Components/SceneComponent.h"
+#include "Components/StaticMeshComponent.h"
+#include "Components/SkeletalMeshComponent.h"
+#include "Camera/CameraComponent.h"
+#include "Components/DirectionalLightComponent.h"
+#include "Camera/PlayerCameraManager.h"
+
+
+/** Get a list of all featured native class types */
+TArray<FNewClassInfo> FFeaturedClasses::AllNativeClasses()
+{
+	TArray<FNewClassInfo> Array;
+	Array.Add(FNewClassInfo(FNewClassInfo::EClassType::EmptyCpp));
+
+	Array += ActorClasses();
+
+	Array.Add(FNewClassInfo(UBlueprintFunctionLibrary::StaticClass()));
+
+	// Add the extra non-UObject classes
+	Array.Add(FNewClassInfo(FNewClassInfo::EClassType::SlateWidget));
+	Array.Add(FNewClassInfo(FNewClassInfo::EClassType::SlateWidgetStyle));
+	return Array;
+}
+
+/** Get a list of all featured Actor class types */
+TArray<FNewClassInfo> FFeaturedClasses::ActorClasses()
+{
+	TArray<FNewClassInfo> Array;
+	Array.Add(FNewClassInfo(ACharacter::StaticClass()));
+	Array.Add(FNewClassInfo(APawn::StaticClass()));
+	Array.Add(FNewClassInfo(AActor::StaticClass()));
+	Array.Add(FNewClassInfo(APlayerCameraManager::StaticClass()));
+	Array.Add(FNewClassInfo(APlayerController::StaticClass()));
+	Array.Add(FNewClassInfo(AGameMode::StaticClass()));
+	Array.Add(FNewClassInfo(AWorldSettings::StaticClass()));
+	Array.Add(FNewClassInfo(AHUD::StaticClass()));
+	Array.Add(FNewClassInfo(APlayerState::StaticClass()));
+	Array.Add(FNewClassInfo(AGameState::StaticClass()));
+	return Array;
+}
+
+/** Get a list of all featured class types */
+TArray<FNewClassInfo> FFeaturedClasses::ComponentClasses()
+{
+	TArray<FNewClassInfo> Array;
+	Array.Add(FNewClassInfo(UActorComponent::StaticClass()));
+	Array.Add(FNewClassInfo(USceneComponent::StaticClass()));
+	Array.Add(FNewClassInfo(UStaticMeshComponent::StaticClass()));
+	Array.Add(FNewClassInfo(USkeletalMeshComponent::StaticClass()));
+	Array.Add(FNewClassInfo(UCameraComponent::StaticClass()));
+	Array.Add(FNewClassInfo(UDirectionalLightComponent::StaticClass()));
+	return Array;
+}
\ No newline at end of file
diff --git a/Engine/Source/Editor/GameProjectGeneration/Public/GameProjectGenerationModule.h b/Engine/Source/Editor/GameProjectGeneration/Public/GameProjectGenerationModule.h
--- a/Engine/Source/Editor/GameProjectGeneration/Public/GameProjectGenerationModule.h
+++ b/Engine/Source/Editor/GameProjectGeneration/Public/GameProjectGenerationModule.h
@@ -5,6 +5,7 @@
 
 #include "ModuleInterface.h"
 #include "ModuleDescriptor.h"
+#include "AddToProjectConfig.h"
 
 struct FTemplateCategory;
 
@@ -21,16 +22,6 @@ struct FModuleContextInfo
 	enum EHostType::Type ModuleType;
 };
 
-
-/**
- * Delegate called when code is added to the project.  Passes in the created class name and class path
- * 
- * @param ClassName		The created class name
- * @param ClassPath		The created class path
- * @param ModuleName	The name of the module that the class was added to
- */
-DECLARE_DELEGATE_ThreeParams( FOnCodeAddedToProject, const FString& /*ClassName*/, const FString& /*ClassPath*/, const FString& /*ModuleName*/);
-
 /**
  * Game Project Generation module
  */
@@ -67,27 +58,21 @@ public:
 
 	/** Creates a new class dialog for creating classes based on the passed-in class. */
 	virtual TSharedRef<class SWidget> CreateNewClassDialog(const UClass* InClass);
-
+	
 	/** 
 	 * Opens a dialog to add code files to the current project. 
 	 *
-	 * @param	InParentWindow		The parent window the dialog should use, or null to choose a suitable default parent window (the main frame, if available)
+	 * @param	Config		Dialog configuration options
 	 */
-	virtual void OpenAddCodeToProjectDialog(const TSharedPtr<SWindow>& InParentWindow = nullptr);
+	virtual void OpenAddCodeToProjectDialog(const FAddToProjectConfig& Config = FAddToProjectConfig());
 
 	/** 
-	 * Opens a dialog to add code files to the current project. 
+	 * Opens a dialog to add a new blueprint to the current project. 
 	 *
-	 * @param	InClass			The class we should force the user to use as their base class type, or null to allow the user to choose their base class in the UI
-	 * @param	InInitialPath		The initial path we should use as the destination for the new class header file, or an empty string to choose a suitable default based upon the module path
-	 * @param	InParentWindow		The parent window the dialog should use, or null to choose a suitable default parent window (the main frame, if available)
-	 * @param	bModal			True if the window should be modal and force the user to make a decision before continuing or false to let the user proceed with other tasks while the window is open
-	 * @param	OnCodeAddedToProject	Callback for when code is successfully added to the project
-	 * @param	InDefaultClassPrefix	Optional argument that specifies the prefix for the new class name.  The user will be able to type their own name if they don't like this name.  Defaults to "My" if not specified or empty
-	 * @param	InDefaultClassName	Optional argument that specifies the default name for the new class being added.  The user will be able to type their own name if they don't like this name.  If empty, defaults to the name of the inherited class.
+	 * @param	Config		Dialog configuration options
 	 */
-	virtual void OpenAddCodeToProjectDialog(const UClass* InClass, const FString& InInitialPath, const TSharedPtr<SWindow>& InParentWindow, bool bModal = false, FOnCodeAddedToProject OnCodeAddedToProject = FOnCodeAddedToProject(), const FString InDefaultClassPrefix = FString(), const FString InDefaultClassName = FString() );
-
+	virtual void OpenAddBlueprintToProjectDialog(const FAddToProjectConfig& Config);
+	
 	/** Delegate for when the AddCodeToProject dialog is opened */
 	DECLARE_EVENT(FGameProjectGenerationModule, FAddCodeToProjectDialogOpenedEvent);
 	FAddCodeToProjectDialogOpenedEvent& OnAddCodeToProjectDialogOpened() { return AddCodeToProjectDialogOpenedEvent; }
diff --git a/Engine/Source/Editor/Kismet/Private/SSCSEditor.cpp b/Engine/Source/Editor/Kismet/Private/SSCSEditor.cpp
--- a/Engine/Source/Editor/Kismet/Private/SSCSEditor.cpp
+++ b/Engine/Source/Editor/Kismet/Private/SSCSEditor.cpp
@@ -7,6 +7,7 @@
 #include "Editor/UnrealEd/Public/Kismet2/ComponentEditorUtils.h"
 #include "BlueprintUtilities.h"
 #include "ComponentAssetBroker.h"
+#include "Editor/ClassViewer/Public/ClassViewerFilter.h"
 
 #include "SSCSEditor.h"
 #include "SKismetInspector.h"
@@ -42,6 +43,8 @@
 #include "NotificationManager.h"
 
 #include "GameProjectGenerationModule.h"
+#include "FeaturedClasses.inl"
+
 #include "HotReloadInterface.h"
 #include "AssetRegistryModule.h"
 #include "SCreateAssetFromObject.h"
@@ -4218,23 +4221,23 @@ TSharedPtr<FSCSEditorTreeNode> SSCSEditor::AddRootComponentTreeNode(UActorCompon
 	return NewTreeNode;
 }
 
-
-void SSCSEditor::MakeDefaultComponentPrefixAndName( TSubclassOf<UActorComponent> ComponentClass, FString& DefaultClassPrefix, FString& DefaultClassName ) const
+class FClassParentFilter : public IClassViewerFilter
 {
-	// If the class being created inherits directly from UActorComponent, we don't want the default new class name to
-	// be called "NewActorComponent", because this will shw up as "New Actor" in the SCS editor tree which doesn't make
-	// a whole lot of sense in the context of a list of components inside an actor.
-	DefaultClassName = ComponentClass->GetName();
-	if( ComponentClass == UActorComponent::StaticClass() )
+public:
+	FClassParentFilter(const TSubclassOf<UActorComponent>& InComponentClass) : ComponentClass(InComponentClass) {}
+
+	virtual bool IsClassAllowed(const FClassViewerInitializationOptions& InInitOptions, const UClass* InClass, TSharedRef< FClassViewerFilterFuncs > InFilterFuncs ) override
 	{
-		// NOTE: This is intentionally not localized because it will become a C++ source file and class name
-		DefaultClassName = TEXT( "Component" );
+		return InClass->IsChildOf(ComponentClass);
 	}
 
-	// NOTE: This is intentionally not localized because it will become a C++ source file and class name
-	DefaultClassPrefix = TEXT( "New" );
-}
+	virtual bool IsUnloadedClassAllowed(const FClassViewerInitializationOptions& InInitOptions, const TSharedRef< const IUnloadedBlueprintData > InUnloadedClassData, TSharedRef< FClassViewerFilterFuncs > InFilterFuncs) override
+	{
+		return InUnloadedClassData->IsChildOf(ComponentClass);
+	}
 
+	TSubclassOf<UActorComponent> ComponentClass;
+};
 
 UClass* SSCSEditor::CreateNewCPPComponent( TSubclassOf<UActorComponent> ComponentClass )
 {
@@ -4249,12 +4252,17 @@ UClass* SSCSEditor::CreateNewCPPComponent( TSubclassOf<UActorComponent> Componen
 		}
 	};
 
-	const bool bModal = true;
-
-	FString DefaultClassPrefix, DefaultClassName;
-	MakeDefaultComponentPrefixAndName( ComponentClass, DefaultClassPrefix, DefaultClassName );
+	FGameProjectGenerationModule::Get().OpenAddCodeToProjectDialog(
+		FAddToProjectConfig()
+		.WindowTitle(LOCTEXT("AddNewC++Component", "Add C++ Component"))
+		.ParentWindow(ParentWindow)
+		.Modal()
+		.OnAddedToProject(FOnAddedToProject::CreateLambda(OnCodeAddedToProject))
+		.FeatureComponentClasses()
+		.AllowableParents(MakeShareable( new FClassParentFilter(ComponentClass) ))
+		.DefaultClassPrefix(TEXT("New"))
+	);
 
-	FGameProjectGenerationModule::Get().OpenAddCodeToProjectDialog(ComponentClass, FString(), ParentWindow, bModal, FOnCodeAddedToProject::CreateLambda(OnCodeAddedToProject), DefaultClassPrefix, DefaultClassName);
 
 	return LoadClass<UActorComponent>(nullptr, *AddedClassName, nullptr, LOAD_None, nullptr);
 }
@@ -4263,63 +4271,37 @@ UClass* SSCSEditor::CreateNewBPComponent(TSubclassOf<UActorComponent> ComponentC
 {
 	UClass* NewClass = nullptr;
 
-	TSharedPtr<SWindow> PickBlueprintPathWidget;
-	SAssignNew(PickBlueprintPathWidget, SWindow)
-		.Title(LOCTEXT("SelectPath", "Select Path"))
-		.ToolTipText(LOCTEXT("SelectPathTooltip", "Select the path where the Blueprint will be created"))
-		.ClientSize(FVector2D(400, 400));
-
-
-	FString PackagePath;
-	FString AssetName;
-	auto OnPathPicked = [&PackagePath, &AssetName](const FString& Path)
-	{
-		PackagePath = Path;
-		AssetName = FPackageName::GetLongPackageAssetName(Path);
-	};
-
-	FString DefaultClassPrefix, DefaultClassName;
-	MakeDefaultComponentPrefixAndName( ComponentClass, DefaultClassPrefix, DefaultClassName );
-
-	TSharedPtr<SCreateAssetFromObject> CreateBlueprintFromActorDialog;
-	PickBlueprintPathWidget->SetContent
-	(
-		SAssignNew(CreateBlueprintFromActorDialog, SCreateAssetFromObject, PickBlueprintPathWidget)
-		.HeadingText(LOCTEXT("CreateBlueprintFromActor_Heading", "Blueprint Name"))
-		.CreateButtonText(LOCTEXT("CreateBlueprintFromActor_ButtonLabel", "Create Blueprint"))
-		.DefaultNameOverride(FText::FromString(DefaultClassPrefix + DefaultClassName))
-		.OnCreateAssetAction(FOnPathChosen::CreateLambda(OnPathPicked))
-	);
-
-
-	FSlateApplication::Get().AddModalWindow(PickBlueprintPathWidget.ToSharedRef(), AsShared());
-
-	if(!PackagePath.IsEmpty() && !AssetName.IsEmpty())
+	auto OnAddedToProject = [&](const FString& ClassName, const FString& PackagePath, const FString& ModuleName)
 	{
-		UPackage* Package = CreatePackage(NULL, *PackagePath);
-
-		if(Package)
+		if(!ClassName.IsEmpty() && !PackagePath.IsEmpty())
 		{
-			// Create and init a new Blueprint
-			if(UBlueprint* NewBP = FKismetEditorUtilities::CreateBlueprint(ComponentClass, Package, FName(*AssetName), BPTYPE_Normal, UBlueprint::StaticClass(), UBlueprintGeneratedClass::StaticClass()))
+			if (UPackage* Package = FindPackage(nullptr, *PackagePath))
 			{
-				// Notify the asset registry
-				FAssetRegistryModule::AssetCreated(NewBP);
-
-				// Mark the package dirty...
-				Package->MarkPackageDirty();
-
-				NewClass = NewBP->GeneratedClass;
+				if (UBlueprint* NewBP = FindObjectFast<UBlueprint>(Package, *ClassName))	
+				{
+					NewClass = NewBP->GeneratedClass;
 
-				TArray<UObject*> Objects;
-				Objects.Add(NewBP);
-				GEditor->SyncBrowserToObjects( Objects );
+					TArray<UObject*> Objects;
+					Objects.Emplace(NewBP);
+					GEditor->SyncBrowserToObjects(Objects);
 
-				// Open the editor for the new blueprint
-				FAssetEditorManager::Get().OpenEditorForAsset(NewBP);
+					// Open the editor for the new blueprint
+					FAssetEditorManager::Get().OpenEditorForAsset(NewBP);
+				}
 			}
 		}
-	}
+	};
+
+	FGameProjectGenerationModule::Get().OpenAddBlueprintToProjectDialog(
+		FAddToProjectConfig()
+		.WindowTitle(LOCTEXT("AddNewBlueprintComponent", "Add Blueprint Component"))
+		.ParentWindow(FSlateApplication::Get().FindWidgetWindow(SharedThis(this)))
+		.Modal()
+		.AllowableParents(MakeShareable( new FClassParentFilter(ComponentClass) ))
+		.FeatureComponentClasses()
+		.OnAddedToProject(FOnAddedToProject::CreateLambda(OnAddedToProject))
+		.DefaultClassPrefix(TEXT("New"))
+	);
 
 	return NewClass;
 }
diff --git a/Engine/Source/Editor/Kismet/Public/SSCSEditor.h b/Engine/Source/Editor/Kismet/Public/SSCSEditor.h
--- a/Engine/Source/Editor/Kismet/Public/SSCSEditor.h
+++ b/Engine/Source/Editor/Kismet/Public/SSCSEditor.h
@@ -999,9 +999,6 @@ protected:
 	/** Adds a root component tree node */
 	TSharedPtr<FSCSEditorTreeNode> AddRootComponentTreeNode(UActorComponent* ActorComp);
 
-	/** Given a class, makes up a default prefix and name for a newly created derived class */
-	void MakeDefaultComponentPrefixAndName( TSubclassOf<UActorComponent> ComponentClass, FString& DefaultClassPrefix, FString& DefaultClassName ) const;
-
 	/**
 	 * Creates a new C++ component from the specified class type
 	 * The user will be prompted to pick a new subclass name and code will be recompiled
diff --git a/Engine/Source/Editor/UnrealEd/Private/Kismet2/ComponentEditorUtils.cpp b/Engine/Source/Editor/UnrealEd/Private/Kismet2/ComponentEditorUtils.cpp
--- a/Engine/Source/Editor/UnrealEd/Private/Kismet2/ComponentEditorUtils.cpp
+++ b/Engine/Source/Editor/UnrealEd/Private/Kismet2/ComponentEditorUtils.cpp
@@ -221,12 +221,19 @@ FString FComponentEditorUtils::GenerateValidVariableName(TSubclassOf<UActorCompo
 	FString ComponentTypeName = FBlueprintEditorUtils::GetClassNameWithoutSuffix(ComponentClass);
 
 	// Strip off 'Component' if the class ends with that.  It just looks better in the UI.
-	const FString SuffixToStrip( TEXT( "Component" ) );
+	FString SuffixToStrip( TEXT( "Component" ) );
 	if( ComponentTypeName.EndsWith( SuffixToStrip ) )
 	{
 		ComponentTypeName = ComponentTypeName.Left( ComponentTypeName.Len() - SuffixToStrip.Len() );
 	}
-	
+
+	// Strip off 'Actor' if the class ends with that so as not to confuse actors with components
+	SuffixToStrip = TEXT( "Actor" );
+	if( ComponentTypeName.EndsWith( SuffixToStrip ) )
+	{
+		ComponentTypeName = ComponentTypeName.Left( ComponentTypeName.Len() - SuffixToStrip.Len() );
+	}
+
 	// Try to create a name without any numerical suffix first
 	int32 Counter = 1;
 	FString ComponentInstanceName = ComponentTypeName;
