[INTEGRATE] Change 2437237 by bruce.nesbit@bnesbit_release_port_edit on 2015/02/08 04:38:39

	switch class/asset position in new project dialog

diff --git a/Engine/Source/Editor/GameProjectGeneration/Private/SNewProjectWizard.cpp b/Engine/Source/Editor/GameProjectGeneration/Private/SNewProjectWizard.cpp
--- a/Engine/Source/Editor/GameProjectGeneration/Private/SNewProjectWizard.cpp
+++ b/Engine/Source/Editor/GameProjectGeneration/Private/SNewProjectWizard.cpp
@@ -495,54 +495,54 @@ void SNewProjectWizard::Construct( const FArguments& InArgs )
 												.Text(this, &SNewProjectWizard::GetSelectedTemplateProperty<FText>, &FTemplateItem::Description)
 											]
 											
-											// Class types
-											+SVerticalBox::Slot()
+											// Asset types
+											+ SVerticalBox::Slot()
 											.AutoHeight()
 											.Padding(FMargin(0.0f, 5.0f, 0.0f, 5.0f))
 											[
 												SNew(SBox)
-												.Visibility(this, &SNewProjectWizard::GetSelectedTemplateClassVisibility)
+												.Visibility(this, &SNewProjectWizard::GetSelectedTemplateAssetVisibility)
 												[
 													SNew(SVerticalBox)
-													+SVerticalBox::Slot()
+													+ SVerticalBox::Slot()
 													[
 														SNew(STextBlock)
 														.TextStyle(FEditorStyle::Get(), "GameProjectDialog.FeatureText")
-														.Text(LOCTEXT("ProjectTemplateClassTypes", "Class Type References:"))
+														.Text(LOCTEXT("ProjectTemplateAssetTypes", "Asset Type References:"))
 													]
 													+ SVerticalBox::Slot()
 													.AutoHeight()
 													[
 														SNew(STextBlock)
 														.AutoWrapText(true)
-														.Text(this, &SNewProjectWizard::GetSelectedTemplateClassTypes)
+														.Text(this, &SNewProjectWizard::GetSelectedTemplateAssetTypes)
 													]
 												]
-												
 											]
-											// Asset types
+											// Class types
 											+ SVerticalBox::Slot()
 											.AutoHeight()
 											.Padding(FMargin(0.0f, 5.0f, 0.0f, 5.0f))
 											[
 												SNew(SBox)
-												.Visibility(this, &SNewProjectWizard::GetSelectedTemplateAssetVisibility)
+												.Visibility(this, &SNewProjectWizard::GetSelectedTemplateClassVisibility)
 												[
 													SNew(SVerticalBox)
 													+ SVerticalBox::Slot()
 													[
 														SNew(STextBlock)
 														.TextStyle(FEditorStyle::Get(), "GameProjectDialog.FeatureText")
-														.Text(LOCTEXT("ProjectTemplateAssetTypes", "Asset Type References:"))
+														.Text(LOCTEXT("ProjectTemplateClassTypes", "Class Type References:"))
 													]
 													+ SVerticalBox::Slot()
-													.AutoHeight()
-													[
-														SNew(STextBlock)
-														.AutoWrapText(true)
-														.Text(this, &SNewProjectWizard::GetSelectedTemplateAssetTypes)
-													]
+														.AutoHeight()
+														[
+															SNew(STextBlock)
+															.AutoWrapText(true)
+															.Text(this, &SNewProjectWizard::GetSelectedTemplateClassTypes)
+														]
 												]
+
 											]
 										]
 									]
