[INTEGRATE] Change 2437737 by [REDACTED_PERSON][REDACTED_HANDLE].7 on 2015/02/09 05:49:56

	Fixed visual studio not opening when creating a code project

	This was a regression caused by 2428764, and addresses UE-9212 - Creating a code project does not bring up IDE

diff --git a/Engine/Source/Editor/GameProjectGeneration/Private/SNewProjectWizard.cpp b/Engine/Source/Editor/GameProjectGeneration/Private/SNewProjectWizard.cpp
--- a/Engine/Source/Editor/GameProjectGeneration/Private/SNewProjectWizard.cpp
+++ b/Engine/Source/Editor/GameProjectGeneration/Private/SNewProjectWizard.cpp
@@ -1470,40 +1470,44 @@ bool SNewProjectWizard::CreateProject( const FString& ProjectFile )
 
 void SNewProjectWizard::CreateAndOpenProject( )
 {
-	if( IsCreateProjectEnabled() )
+	if( !IsCreateProjectEnabled() )
 	{
-		FString ProjectFile = GetProjectFilenameWithPath();
-		if ( CreateProject(ProjectFile) )
-		{
-			// Prevent periodic validity checks. This is to prevent a brief error message about the project already existing while you are exiting.
-			bPreventPeriodicValidityChecksUntilNextChange = true;
+		return;
+	}
+
+	FString ProjectFile = GetProjectFilenameWithPath();
+	if ( !CreateProject(ProjectFile) )
+	{
+		return;
+	}
+
+	// Prevent periodic validity checks. This is to prevent a brief error message about the project already existing while you are exiting.
+	bPreventPeriodicValidityChecksUntilNextChange = true;
 
-			bool bCanOpenProject = false;
-			if( GetSelectedTemplateItem()->bGenerateCode )
+	if( GetSelectedTemplateItem()->bGenerateCode )
+	{
+	    // Rocket already has the engine compiled, so we can try to build and open a new project immediately. Non-Rocket might require building
+	    // the engine (especially the case when binaries came from P4), so we only open the IDE for that.
+		if (FRocketSupport::IsRocket())
+		{
+			if (GameProjectUtils::BuildCodeProject(ProjectFile))
 			{
-			    // Rocket already has the engine compiled, so we can try to build and open a new project immediately. Non-Rocket might require building
-			    // the engine (especially the case when binaries came from P4), so we only open the IDE for that.
-			    if( FRocketSupport::IsRocket() && GameProjectUtils::BuildCodeProject(ProjectFile) )
-			    {
-					// Everything compiled OK, so we can go ahead and open the project
-					bCanOpenProject = true;
-		
-					// Open Visual Studio or Xcode if the user created a project with C++ files.  Note that if the code failed to compile, the
-					// BuildCodeProject() function will already offer to open the IDE for the user, so we only do this if every compiled OK.
-					OpenCodeIDE( ProjectFile );
-			    }
+				OpenCodeIDE( ProjectFile );
+				OpenProject( ProjectFile );
 			}
 			else
 			{
-				// We can always open non-code projects, because they don't have any DLLs
-				bCanOpenProject = true;
-			}
-
-			if( bCanOpenProject )
-			{
-				OpenProject(ProjectFile);
+				// User will have already been prompted to open the IDE
 			}
 		}
+		else
+		{
+			OpenCodeIDE( ProjectFile );
+		}
+	}
+	else
+	{
+		OpenProject( ProjectFile );
 	}
 }
 
