[INTEGRATE] Change 2437865 by bruce.nesbit@bnesbit_Releases on 2015/02/09 08:32:02

	Fixed default to including starter content
	revised code that creattes starter content combo

diff --git a/Engine/Config/BaseEditorGameAgnostic.ini b/Engine/Config/BaseEditorGameAgnostic.ini
--- a/Engine/Config/BaseEditorGameAgnostic.ini
+++ b/Engine/Config/BaseEditorGameAgnostic.ini
@@ -1,4 +1,5 @@
 [/Script/UnrealEd.EditorGameAgnosticSettings]
+bCopyStarterContentPreference = false
 
 [/Script/NewsFeed.NewsFeedSettings]
 CdnSourceUrl="http://cdn.unrealengine.com/newsfeed/[[Culture]]/index.json"
diff --git a/Engine/Source/Editor/GameProjectGeneration/Private/SNewProjectWizard.cpp b/Engine/Source/Editor/GameProjectGeneration/Private/SNewProjectWizard.cpp
--- a/Engine/Source/Editor/GameProjectGeneration/Private/SNewProjectWizard.cpp
+++ b/Engine/Source/Editor/GameProjectGeneration/Private/SNewProjectWizard.cpp
@@ -330,22 +330,15 @@ void SNewProjectWizard::Construct( const FArguments& InArgs )
 		TArray<SDecoratedEnumCombo<int32>::FComboOption> StarterContentInfo;
 		StarterContentInfo.Add(SDecoratedEnumCombo<int32>::FComboOption(
 			0, FSlateIcon(FEditorStyle::GetStyleSetName(), "GameProjectDialog.NoStarterContent"), LOCTEXT("NoStarterContent", "No Starter Content")));
-		FString RequiredTooltipText = LOCTEXT("CopyStarterContent_ToolTip", "Enable to include an additional content pack containing simple placeable meshes with basic materials and textures.\nYou can opt out of including this to create a project that only has the bare essentials for the selected project template.").ToString();
-		// Only add the option to add starter content if its there to add !
-		if( GameProjectUtils::IsStarterContentAvailableForNewProjects() == true )
-		{
-			StarterContentInfo.Add(SDecoratedEnumCombo<int32>::FComboOption(
-				1, FSlateIcon(FEditorStyle::GetStyleSetName(), "GameProjectDialog.IncludeStarterContent"), LOCTEXT("IncludeStarterContent", "With Starter Content")));
-		}
-		else
-		{
-			RequiredTooltipText = LOCTEXT("NoStarterContent_ToolTip", "Starter content is not currently available.").ToString();
-		}
 
+		// Only add the option to add starter content if its there to add !
+		bool bIsStarterAvailable = GameProjectUtils::IsStarterContentAvailableForNewProjects();
+		StarterContentInfo.Add(SDecoratedEnumCombo<int32>::FComboOption(
+			1, FSlateIcon(FEditorStyle::GetStyleSetName(), "GameProjectDialog.IncludeStarterContent"), LOCTEXT("IncludeStarterContent", "With Starter Content"),bIsStarterAvailable));
 		StartContentCombo = SNew(SDecoratedEnumCombo<int32>, MoveTemp(StarterContentInfo))
 			.SelectedEnum(this, &SNewProjectWizard::GetCopyStarterContentIndex)
 			.OnEnumChanged(this, &SNewProjectWizard::OnSetCopyStarterContent)
-			.ToolTipText( RequiredTooltipText );
+			.ToolTipText(LOCTEXT("CopyStarterContent_ToolTip", "Enable to include an additional content pack containing simple placeable meshes with basic materials and textures.\nYou can opt out of including this to create a project that only has the bare essentials for the selected project template.").ToString());
 	}
 
 	const float UniformPadding = 16.f;
