[INTEGRATE] Change 2438352 by [REDACTED_PERSON][REDACTED_HANDLE] on 2015/02/09 12:44:42

	[UE-9373] - Button/hyperlink to get the suggested IDE should now only be visible when no compiler is available

diff --git a/Engine/Source/Editor/GameProjectGeneration/Private/SGetSuggestedIDEWidget.cpp b/Engine/Source/Editor/GameProjectGeneration/Private/SGetSuggestedIDEWidget.cpp
--- a/Engine/Source/Editor/GameProjectGeneration/Private/SGetSuggestedIDEWidget.cpp
+++ b/Engine/Source/Editor/GameProjectGeneration/Private/SGetSuggestedIDEWidget.cpp
@@ -39,8 +39,7 @@ TSharedRef<SWidget> SGetSuggestedIDEWidget::CreateGetSuggestedIDEWidget() const
 
 EVisibility SGetSuggestedIDEWidget::GetVisibility() const
 {
-	/*return FSourceCodeNavigation::IsCompilerAvailable() ? EVisibility::Collapsed : EVisibility::Visible;*/
-	return EVisibility::Visible;
+	return FSourceCodeNavigation::IsCompilerAvailable() ? EVisibility::Collapsed : EVisibility::Visible;
 }
 
 void SGetSuggestedIDEWidget::OnDownloadIDEClicked(FString URL)
