[INTEGRATE] Change 2438567 by [REDACTED_PERSON][REDACTED_HANDLE] on 2015/02/09 14:54:19

	Editor - The new project dialog has gotten a bit of a visual facelift for the buttons.

diff --git a/Engine/Source/Editor/GameProjectGeneration/Private/SNewProjectWizard.cpp b/Engine/Source/Editor/GameProjectGeneration/Private/SNewProjectWizard.cpp
--- a/Engine/Source/Editor/GameProjectGeneration/Private/SNewProjectWizard.cpp
+++ b/Engine/Source/Editor/GameProjectGeneration/Private/SNewProjectWizard.cpp
@@ -354,6 +354,11 @@ void SNewProjectWizard::Construct( const FArguments& InArgs )
 			.Padding(UniformPadding / 2.0f)
 			[
 				SAssignNew(MainWizard, SWizard)
+				.ButtonStyle(FEditorStyle::Get(), "FlatButton.Default")
+				.CancelButtonStyle(FEditorStyle::Get(), "FlatButton.Warning")
+				.FinishButtonStyle(FEditorStyle::Get(), "FlatButton.Success")
+				.ButtonTextStyle(FEditorStyle::Get(), "LargeText")
+				.ForegroundColor(FEditorStyle::Get().GetSlateColor("WhiteBrush"))
 				.ShowPageList(false)
 				.ShowCancelButton(false)
 				.CanFinish(this, &SNewProjectWizard::HandleCreateProjectWizardCanFinish)
@@ -379,7 +384,7 @@ void SNewProjectWizard::Construct( const FArguments& InArgs )
 							SNew(SRichTextBlock)
 							.Text(LOCTEXT("ProjectTemplateDescription", "Choose a <RichTextBlock.BoldHighlight>template</> to use as a starting point for your new project.  Any of these features can be added later by clicking <RichTextBlock.BoldHighlight>Add Feature or Content Pack</> in <RichTextBlock.BoldHighlight>Content Browser</>."))
 							.AutoWrapText(true)
- 							.DecoratorStyleSet(&FEditorStyle::Get())
+							.DecoratorStyleSet(&FEditorStyle::Get())
 							.ToolTip(IDocumentation::Get()->CreateToolTip(LOCTEXT("TemplateChoiceTooltip", "A template consists of a little bit of player control logic (either as a Blueprint or in C++), input bindings, and appropriate prototyping assets."), NULL, TEXT("Shared/Editor/NewProjectWizard"), TEXT("TemplateChoice")))
 						]
 
@@ -561,7 +566,7 @@ void SNewProjectWizard::Construct( const FArguments& InArgs )
 									SNew(SRichTextBlock)
 									.Text(LOCTEXT("ProjectSettingsDescription", "Choose some <RichTextBlock.BoldHighlight>settings</> for your project.  Don't worry, you can change these later in the <RichTextBlock.BoldHighlight>Target Hardware</> section of <RichTextBlock.BoldHighlight>Project Settings</>.  You can also add the <RichTextBlock.BoldHighlight>Starter Content</> to your project later using <RichTextBlock.BoldHighlight>Content Browser</>."))
 									.AutoWrapText(true)
- 									.DecoratorStyleSet(&FEditorStyle::Get())
+									.DecoratorStyleSet(&FEditorStyle::Get())
 									.ToolTip(IDocumentation::Get()->CreateToolTip(LOCTEXT("HardwareTargetTooltip", "These settings will choose good defaults for a number of other settings in the project such as post-processing flags and touch input emulation using the mouse."), NULL, TEXT("Shared/Editor/NewProjectWizard"), TEXT("TargetHardware")))
 								]
 
@@ -633,7 +638,7 @@ void SNewProjectWizard::Construct( const FArguments& InArgs )
 									SNew(SRichTextBlock)
 									.Text(LOCTEXT("ProjectPathDescription", "Select a <RichTextBlock.BoldHighlight>location</> for your project to be stored."))
 									.AutoWrapText(true)
- 									.DecoratorStyleSet(&FEditorStyle::Get())
+									.DecoratorStyleSet(&FEditorStyle::Get())
 									.ToolTip(IDocumentation::Get()->CreateToolTip(LOCTEXT("ProjectPathDescriptionTooltip", "All of your project content and code will be stored here."), NULL, TEXT("Shared/Editor/NewProjectWizard"), TEXT("ProjectPath")))
 								]
 
