[INTEGRATE] Change 2443405 by [REDACTED_PERSON][REDACTED_HANDLE] on 2015/02/12 09:48:14

	C++ template improvements for "Add Code" editor feature.

diff --git a/Engine/Source/Editor/GameProjectGeneration/Private/GameProjectUtils.cpp b/Engine/Source/Editor/GameProjectGeneration/Private/GameProjectUtils.cpp
--- a/Engine/Source/Editor/GameProjectGeneration/Private/GameProjectUtils.cpp
+++ b/Engine/Source/Editor/GameProjectGeneration/Private/GameProjectUtils.cpp
@@ -251,19 +251,27 @@ FString FNewClassInfo::GetHeaderTemplateFilename() const
 	{
 		case EClassType::UObject:
 		{
-			if( BaseClass != nullptr && ( BaseClass == UActorComponent::StaticClass() || BaseClass == USceneComponent::StaticClass() ) )
+			if (BaseClass != nullptr)
 			{
-				return TEXT( "ActorComponentClass.h.template" );
-			}
-			else if( BaseClass != nullptr && BaseClass == AActor::StaticClass() )
-			{
-				return TEXT( "ActorClass.h.template" );
-			}
-			else
-			{
-				// Some other non-actor, non-component UObject class
-				return TEXT( "UObjectClass.h.template" );
+				if ((BaseClass == UActorComponent::StaticClass()) || (BaseClass == USceneComponent::StaticClass()))
+				{
+					return TEXT("ActorComponentClass.h.template");
+				}
+				else if (BaseClass == AActor::StaticClass())
+				{
+					return TEXT("ActorClass.h.template");
+				}
+				else if (BaseClass == APawn::StaticClass())
+				{
+					return TEXT("PawnClass.h.template");
+				}
+				else if (BaseClass == ACharacter::StaticClass())
+				{
+					return TEXT("CharacterClass.h.template");
+				}
 			}
+			// Some other non-actor, non-component UObject class
+			return TEXT( "UObjectClass.h.template" );
 		}
 
 	case EClassType::EmptyCpp:
@@ -286,19 +294,27 @@ FString FNewClassInfo::GetSourceTemplateFilename() const
 	switch(ClassType)
 	{
 		case EClassType::UObject:
-			if( BaseClass != nullptr && ( BaseClass == UActorComponent::StaticClass() || BaseClass == USceneComponent::StaticClass() ) )
+			if (BaseClass != nullptr)
 			{
-				return TEXT( "ActorComponentClass.cpp.template" );
-			}
-			else if( BaseClass != nullptr && BaseClass == AActor::StaticClass() )
-			{
-				return TEXT( "ActorClass.cpp.template" );
-			}
-			else
-			{
-				// Some other non-actor, non-component UObject class
-				return TEXT( "UObjectClass.cpp.template" );
+				if ((BaseClass == UActorComponent::StaticClass()) || (BaseClass == USceneComponent::StaticClass()))
+				{
+					return TEXT("ActorComponentClass.cpp.template");
+				}
+				else if (BaseClass == AActor::StaticClass())
+				{
+					return TEXT("ActorClass.cpp.template");
+				}
+				else if (BaseClass == APawn::StaticClass())
+				{
+					return TEXT("PawnClass.cpp.template");
+				}
+				else if (BaseClass == ACharacter::StaticClass())
+				{
+					return TEXT("CharacterClass.cpp.template");
+				}
 			}
+			// Some other non-actor, non-component UObject class
+			return TEXT( "UObjectClass.cpp.template" );
 	
 	case EClassType::EmptyCpp:
 		return TEXT("EmptyClass.cpp.template");
diff --git a/Engine/Source/Runtime/Engine/Classes/GameFramework/Character.h b/Engine/Source/Runtime/Engine/Classes/GameFramework/Character.h
--- a/Engine/Source/Runtime/Engine/Classes/GameFramework/Character.h
+++ b/Engine/Source/Runtime/Engine/Classes/GameFramework/Character.h
@@ -201,11 +201,10 @@ class ENGINE_API ACharacter : public APawn
 {
 	GENERATED_BODY()
 public:
-
 	/**
 	 * Default UObject constructor.
 	 */
-	ACharacter(const FObjectInitializer& ObjectInitializer);
+	ACharacter(const FObjectInitializer& ObjectInitializer = FObjectInitializer::Get());
 
 	void GetLifetimeReplicatedProps(TArray<FLifetimeProperty>& OutLifetimeProps) const override;
 private_subobject:
diff --git a/Engine/Source/Runtime/Engine/Classes/GameFramework/Pawn.h b/Engine/Source/Runtime/Engine/Classes/GameFramework/Pawn.h
--- a/Engine/Source/Runtime/Engine/Classes/GameFramework/Pawn.h
+++ b/Engine/Source/Runtime/Engine/Classes/GameFramework/Pawn.h
@@ -32,11 +32,10 @@ class ENGINE_API APawn : public AActor, public INavAgentInterface
 {
 	GENERATED_BODY()
 public:
-
 	/**
-	 * Default UObject constructor.
-	 */
-	APawn(const FObjectInitializer& ObjectInitializer);
+	* Default UObject constructor.
+	*/
+	APawn(const FObjectInitializer& ObjectInitializer = FObjectInitializer::Get());
 
 	virtual void GetLifetimeReplicatedProps(TArray<FLifetimeProperty>& OutLifetimeProps) const;
 
