fix for UE-9622 Refresh button doesn't work in Open Project

diff --git a/Engine/Source/Editor/GameProjectGeneration/Private/SProjectBrowser.cpp b/Engine/Source/Editor/GameProjectGeneration/Private/SProjectBrowser.cpp
--- a/Engine/Source/Editor/GameProjectGeneration/Private/SProjectBrowser.cpp
+++ b/Engine/Source/Editor/GameProjectGeneration/Private/SProjectBrowser.cpp
@@ -738,7 +738,9 @@ FReply SProjectBrowser::FindProjects()
 	}
 
 	// Add all the native project files we can find, and automatically filter them depending on their directory
-	const TArray<FString> &NativeProjectFiles = FUProjectDictionary::GetDefault().GetProjectPaths();
+	FUProjectDictionary& DefaultProjectDictionary = FUProjectDictionary::GetDefault();
+	DefaultProjectDictionary.Refresh();
+	const TArray<FString> &NativeProjectFiles = DefaultProjectDictionary.GetProjectPaths();
 	for(int32 Idx = 0; Idx < NativeProjectFiles.Num(); Idx++)
 	{
 		if(!NativeProjectFiles[Idx].Contains(TEXT("/Templates/")))
