Cleaned up the parent class tree when adding new component classes

UE-9554 - Add BP/C++ Component's Show all classes list not expanded by default

It now hides the Object root if UObject is not a valid base class type, as this allows the next level of classes to be expanded out which looks much better when adding new component classes.

This involved a lot of shuffling of module dependencies so that we could link GameProjectGeneration to ClassViewer, as previously ClassViewer was directly linking to GameProjectGeneration.

I went through and fixed everything that was linking to GameProjectGeneration to instead list it as a dynamic dependency (which it is, as nothing needs to link to it), and then verified that everything was using ClassViewer correctly (found some places that were only marking it as a dynamic dependency without also adding it as an include module, causing issues with missing API macro definitions as they were including ClassViewer headers via a relative include path).

diff --git a/Engine/Plugins/Editor/PluginsEditor/Source/PluginsEditor/PluginsEditor.Build.cs b/Engine/Plugins/Editor/PluginsEditor/Source/PluginsEditor/PluginsEditor.Build.cs
--- a/Engine/Plugins/Editor/PluginsEditor/Source/PluginsEditor/PluginsEditor.Build.cs
+++ b/Engine/Plugins/Editor/PluginsEditor/Source/PluginsEditor/PluginsEditor.Build.cs
@@ -22,6 +22,17 @@ namespace UnrealBuildTool.Rules
 					"EditorStyle",
 					"Projects",
 					"UnrealEd",
+				}
+			);
+
+			PrivateIncludePathModuleNames.AddRange(
+				new string[] {
+					"GameProjectGeneration",
+				}
+			);
+
+			DynamicallyLoadedModuleNames.AddRange(
+				new string[] {
 					"GameProjectGeneration",
 				}
 			);
diff --git a/Engine/Source/Developer/iOS/IOSPlatformEditor/IOSPlatformEditor.Build.cs b/Engine/Source/Developer/iOS/IOSPlatformEditor/IOSPlatformEditor.Build.cs
--- a/Engine/Source/Developer/iOS/IOSPlatformEditor/IOSPlatformEditor.Build.cs
+++ b/Engine/Source/Developer/iOS/IOSPlatformEditor/IOSPlatformEditor.Build.cs
@@ -15,7 +15,6 @@ public class IOSPlatformEditor : ModuleRules
 				"InputCore",
 				"DesktopPlatform",
 				"Engine",
-				"GameProjectGeneration",
 				"MainFrame",
 				"Slate",
 				"SlateCore",
@@ -30,10 +29,17 @@ public class IOSPlatformEditor : ModuleRules
 
 		PrivateIncludePathModuleNames.AddRange(
 			new string[] {
+				"GameProjectGeneration",
 				"Settings",
 			}
 		);
 
+		DynamicallyLoadedModuleNames.AddRange(
+			new string[] {
+				"GameProjectGeneration",
+				}
+		);
+
         // this is listed above, so it isn't really dynamically loaded, this just marks it as being platform specific.
 		PlatformSpecificDynamicallyLoadedModuleNames.Add("IOSRuntimeSettings");
 	}
diff --git a/Engine/Source/Editor/ClassViewer/ClassViewer.Build.cs b/Engine/Source/Editor/ClassViewer/ClassViewer.Build.cs
--- a/Engine/Source/Editor/ClassViewer/ClassViewer.Build.cs
+++ b/Engine/Source/Editor/ClassViewer/ClassViewer.Build.cs
@@ -10,7 +10,8 @@ public class ClassViewer : ModuleRules
 			new string[] {
 				"AssetRegistry",
 				"EditorWidgets",
-                "PropertyEditor"
+				"GameProjectGeneration",
+                "PropertyEditor",
 			}
 		);
 
@@ -25,7 +26,6 @@ public class ClassViewer : ModuleRules
                 "EditorStyle",
 				"UnrealEd",
 				"PropertyEditor",
-				"GameProjectGeneration"
 			}
 		);
 
@@ -33,7 +33,8 @@ public class ClassViewer : ModuleRules
 			new string[] {
 				"AssetRegistry",
 				"EditorWidgets",
-				"PropertyEditor"
+				"GameProjectGeneration",
+				"PropertyEditor",
 			}
 		);
 	}
diff --git a/Engine/Source/Editor/ClassViewer/Public/ClassViewerFilter.h b/Engine/Source/Editor/ClassViewer/Public/ClassViewerFilter.h
--- a/Engine/Source/Editor/ClassViewer/Public/ClassViewerFilter.h
+++ b/Engine/Source/Editor/ClassViewer/Public/ClassViewerFilter.h
@@ -30,7 +30,7 @@ namespace EFilterReturn
 {
 	enum Type{ Failed = 0, Passed, NoItems };
 }
-class FClassViewerFilterFuncs
+class CLASSVIEWER_API FClassViewerFilterFuncs
 {
 public:
 
diff --git a/Engine/Source/Editor/ContentBrowser/ContentBrowser.Build.cs b/Engine/Source/Editor/ContentBrowser/ContentBrowser.Build.cs
--- a/Engine/Source/Editor/ContentBrowser/ContentBrowser.Build.cs
+++ b/Engine/Source/Editor/ContentBrowser/ContentBrowser.Build.cs
@@ -12,6 +12,7 @@ public class ContentBrowser : ModuleRules
 				"AssetTools",
 				"CollectionManager",
 				"EditorWidgets",
+				"GameProjectGeneration",
                 "MainFrame",
 				"SourceControl",
 				"SourceControlWindows",
@@ -26,7 +27,6 @@ public class ContentBrowser : ModuleRules
 				"CoreUObject",
                 "InputCore",
 				"Engine",
-				"GameProjectGeneration",
 				"Slate",
 				"SlateCore",
                 "EditorStyle",
@@ -47,6 +47,7 @@ public class ContentBrowser : ModuleRules
 				"AssetTools",
 				"CollectionManager",
 				"EditorWidgets",
+				"GameProjectGeneration",
                 "MainFrame",
                 "ReferenceViewer"
 			}
diff --git a/Engine/Source/Editor/DetailCustomizations/DetailCustomizations.Build.cs b/Engine/Source/Editor/DetailCustomizations/DetailCustomizations.Build.cs
--- a/Engine/Source/Editor/DetailCustomizations/DetailCustomizations.Build.cs
+++ b/Engine/Source/Editor/DetailCustomizations/DetailCustomizations.Build.cs
@@ -33,7 +33,6 @@ public class DetailCustomizations : ModuleRules
 				"ClassViewer",
 				"TargetPlatform",
 				"ExternalImagePicker",
-				"GameProjectGeneration",
 				"MoviePlayer",
 				"SourceControl",
                 "InternationalizationSettings",
@@ -53,6 +52,7 @@ public class DetailCustomizations : ModuleRules
 				"Landscape",
 				"LandscapeEditor",
 				"PropertyEditor",
+				"GameProjectGeneration",
 			}
 		);
 
@@ -60,6 +60,7 @@ public class DetailCustomizations : ModuleRules
 			new string[] {
 				"Layers",
 				"PropertyEditor",
+				"GameProjectGeneration",
 			}
 		);
 	}
diff --git a/Engine/Source/Editor/DetailCustomizations/Private/ActorDetails.cpp b/Engine/Source/Editor/DetailCustomizations/Private/ActorDetails.cpp
--- a/Engine/Source/Editor/DetailCustomizations/Private/ActorDetails.cpp
+++ b/Engine/Source/Editor/DetailCustomizations/Private/ActorDetails.cpp
@@ -4,10 +4,10 @@
 #include "ActorDetails.h"
 #include "AssetSelection.h"
 #include "Editor/Layers/Public/LayersModule.h"
-#include "Editor/LevelEditor/Public/LevelEditor.h"
-#include "Editor/LevelEditor/Public/LevelEditorActions.h"
-#include "Editor/ClassViewer/Public/ClassViewerModule.h"
-#include "Editor/ClassViewer/Public/ClassViewerFilter.h"
+#include "LevelEditor.h"
+#include "LevelEditorActions.h"
+#include "ClassViewerModule.h"
+#include "ClassViewerFilter.h"
 #include "Editor/UnrealEd/Public/Kismet2/KismetEditorUtilities.h"
 #include "Editor/UnrealEd/Public/Kismet2/KismetDebugUtilities.h"
 #include "ClassIconFinder.h"
diff --git a/Engine/Source/Editor/DetailCustomizations/Private/SkeletalMeshComponentDetails.cpp b/Engine/Source/Editor/DetailCustomizations/Private/SkeletalMeshComponentDetails.cpp
--- a/Engine/Source/Editor/DetailCustomizations/Private/SkeletalMeshComponentDetails.cpp
+++ b/Engine/Source/Editor/DetailCustomizations/Private/SkeletalMeshComponentDetails.cpp
@@ -4,8 +4,8 @@
 #include "SkeletalMeshComponentDetails.h"
 #include "PropertyCustomizationHelpers.h"
 #include "AssetData.h"
-#include "Editor/ClassViewer/Public/ClassViewerModule.h"
-#include "Editor/ClassViewer/Public/ClassViewerFilter.h"
+#include "ClassViewerModule.h"
+#include "ClassViewerFilter.h"
 #include "Engine/Selection.h"
 #include "Animation/AnimBlueprintGeneratedClass.h"
 
diff --git a/Engine/Source/Editor/GameProjectGeneration/GameProjectGeneration.Build.cs b/Engine/Source/Editor/GameProjectGeneration/GameProjectGeneration.Build.cs
--- a/Engine/Source/Editor/GameProjectGeneration/GameProjectGeneration.Build.cs
+++ b/Engine/Source/Editor/GameProjectGeneration/GameProjectGeneration.Build.cs
@@ -8,7 +8,6 @@ public class GameProjectGeneration : ModuleRules
 	{
         PrivateIncludePathModuleNames.AddRange(
             new string[] {
-				"ClassViewer",
                 "DesktopPlatform",
                 "MainFrame",
             }
@@ -18,6 +17,7 @@ public class GameProjectGeneration : ModuleRules
 			new string[] {
 				"Analytics",
                 "AppFramework",
+				"ClassViewer",
 				"Core",
 				"CoreUObject",
 				"Engine",
@@ -39,7 +39,6 @@ public class GameProjectGeneration : ModuleRules
 
         DynamicallyLoadedModuleNames.AddRange(
             new string[] {
-				"ClassViewer",
                 "DesktopPlatform",
                 "Documentation",
                 "MainFrame",
diff --git a/Engine/Source/Editor/GameProjectGeneration/Private/SNewClassDialog.cpp b/Engine/Source/Editor/GameProjectGeneration/Private/SNewClassDialog.cpp
--- a/Engine/Source/Editor/GameProjectGeneration/Private/SNewClassDialog.cpp
+++ b/Engine/Source/Editor/GameProjectGeneration/Private/SNewClassDialog.cpp
@@ -166,6 +166,7 @@ void SNewClassDialog::Construct( const FArguments& InArgs )
 	Options.bShowUnloadedBlueprints = false;
 	Options.bShowNoneOption = false;
 	Options.bShowObjectRootClass = true;
+	Options.bExpandRootNodes = true;
 
 	if (InArgs._ClassViewerFilter.IsValid())
 	{
@@ -177,6 +178,12 @@ void SNewClassDialog::Construct( const FArguments& InArgs )
 		Options.ClassFilter = MakeShareable(new FNativeClassParentFilter());
 	}
 
+	// Only show the Object root class if it's a valid base (this helps keep the tree clean)
+	if (Options.ClassFilter.IsValid() && !Options.ClassFilter->IsClassAllowed(Options, UObject::StaticClass(), MakeShareable(new FClassViewerFilterFuncs)))
+	{
+		Options.bShowObjectRootClass = false;
+	}
+
 	ClassViewer = StaticCastSharedRef<SClassViewer>(FModuleManager::LoadModuleChecked<FClassViewerModule>("ClassViewer").CreateClassViewer(Options, FOnClassPicked::CreateSP(this, &SNewClassDialog::OnAdvancedClassSelected)));
 
 	// Make sure the featured classes all pass the active class filter
diff --git a/Engine/Source/Editor/GameplayAbilitiesEditor/GameplayAbilitiesEditor.Build.cs b/Engine/Source/Editor/GameplayAbilitiesEditor/GameplayAbilitiesEditor.Build.cs
--- a/Engine/Source/Editor/GameplayAbilitiesEditor/GameplayAbilitiesEditor.Build.cs
+++ b/Engine/Source/Editor/GameplayAbilitiesEditor/GameplayAbilitiesEditor.Build.cs
@@ -23,6 +23,7 @@ namespace UnrealBuildTool.Rules
 					"CoreUObject",
 					"Engine",
 					"AssetTools",
+					"ClassViewer",
                     "GameplayTags",
 					"GameplayAbilities",
                     "InputCore",
diff --git a/Engine/Source/Editor/GameplayAbilitiesEditor/Private/GameplayAbilitiesBlueprintFactory.cpp b/Engine/Source/Editor/GameplayAbilitiesEditor/Private/GameplayAbilitiesBlueprintFactory.cpp
--- a/Engine/Source/Editor/GameplayAbilitiesEditor/Private/GameplayAbilitiesBlueprintFactory.cpp
+++ b/Engine/Source/Editor/GameplayAbilitiesEditor/Private/GameplayAbilitiesBlueprintFactory.cpp
@@ -10,7 +10,7 @@
 #include "GameplayAbilityGraph.h"
 #include "GameplayAbilityGraphSchema.h"
 
-#include "Editor/ClassViewer/Public/ClassViewerFilter.h"
+#include "ClassViewerFilter.h"
 
 #include "GameplayAbilityBlueprintGeneratedClass.h"
 #include "GameplayAbilitiesBlueprintFactory.h"
diff --git a/Engine/Source/Editor/GraphEditor/GraphEditor.Build.cs b/Engine/Source/Editor/GraphEditor/GraphEditor.Build.cs
--- a/Engine/Source/Editor/GraphEditor/GraphEditor.Build.cs
+++ b/Engine/Source/Editor/GraphEditor/GraphEditor.Build.cs
@@ -43,6 +43,7 @@ public class GraphEditor : ModuleRules
                 "EditorStyle",
 				"UnrealEd",
 				"AssetRegistry",
+				"ClassViewer",
 				"KismetWidgets",
 				"BlueprintGraph",
 				"AnimGraph",
diff --git a/Engine/Source/Editor/GraphEditor/Private/KismetNodes/SGraphNodeSpawnActor.cpp b/Engine/Source/Editor/GraphEditor/Private/KismetNodes/SGraphNodeSpawnActor.cpp
--- a/Engine/Source/Editor/GraphEditor/Private/KismetNodes/SGraphNodeSpawnActor.cpp
+++ b/Engine/Source/Editor/GraphEditor/Private/KismetNodes/SGraphNodeSpawnActor.cpp
@@ -7,8 +7,8 @@
 #include "SGraphNodeSpawnActor.h"
 #include "KismetPins/SGraphPinObject.h"
 #include "NodeFactory.h"
-#include "Editor/ClassViewer/Public/ClassViewerModule.h"
-#include "Editor/ClassViewer/Public/ClassViewerFilter.h"
+#include "ClassViewerModule.h"
+#include "ClassViewerFilter.h"
 
 #define LOCTEXT_NAMESPACE "SGraphNodeSpawnActor"
 
diff --git a/Engine/Source/Editor/Kismet/Kismet.Build.cs b/Engine/Source/Editor/Kismet/Kismet.Build.cs
--- a/Engine/Source/Editor/Kismet/Kismet.Build.cs
+++ b/Engine/Source/Editor/Kismet/Kismet.Build.cs
@@ -12,10 +12,12 @@ public class Kismet : ModuleRules
 			new string[] { 
 				"AssetRegistry", 
 				"AssetTools",
+				"ClassViewer",
 				"EditorWidgets",
 				"Analytics",
                 "DerivedDataCache",
                 "LevelEditor",
+				"GameProjectGeneration",
 			}
 			);
 
@@ -45,7 +47,6 @@ public class Kismet : ModuleRules
                 "DerivedDataCache",
 				"DesktopPlatform",
 				"HotReload",
-				"GameProjectGeneration",
 			}
 			);
 
@@ -53,7 +54,8 @@ public class Kismet : ModuleRules
             new string[] {
 				"ClassViewer",
 				"Documentation",
-				"EditorWidgets"
+				"EditorWidgets",
+				"GameProjectGeneration",
 			}
             );
 	}
diff --git a/Engine/Source/Editor/Kismet/Private/SBlueprintLibraryPalette.cpp b/Engine/Source/Editor/Kismet/Private/SBlueprintLibraryPalette.cpp
--- a/Engine/Source/Editor/Kismet/Private/SBlueprintLibraryPalette.cpp
+++ b/Engine/Source/Editor/Kismet/Private/SBlueprintLibraryPalette.cpp
@@ -2,8 +2,8 @@
 
 #include "BlueprintEditorPrivatePCH.h"
 #include "SBlueprintLibraryPalette.h"
-#include "Editor/ClassViewer/Public/ClassViewerModule.h" // for FClassViewerInitializationOptions
-#include "Editor/ClassViewer/Public/ClassViewerFilter.h" // for FClassViewerFilterFuncs
+#include "ClassViewerModule.h" // for FClassViewerInitializationOptions
+#include "ClassViewerFilter.h" // for FClassViewerFilterFuncs
 #include "K2ActionMenuBuilder.h"
 #include "BlueprintPaletteFavorites.h"
 #include "BlueprintActionFilter.h"
diff --git a/Engine/Source/Editor/Kismet/Private/SSCSEditor.cpp b/Engine/Source/Editor/Kismet/Private/SSCSEditor.cpp
--- a/Engine/Source/Editor/Kismet/Private/SSCSEditor.cpp
+++ b/Engine/Source/Editor/Kismet/Private/SSCSEditor.cpp
@@ -7,7 +7,7 @@
 #include "Editor/UnrealEd/Public/Kismet2/ComponentEditorUtils.h"
 #include "BlueprintUtilities.h"
 #include "ComponentAssetBroker.h"
-#include "Editor/ClassViewer/Public/ClassViewerFilter.h"
+#include "ClassViewerFilter.h"
 
 #include "SSCSEditor.h"
 #include "SKismetInspector.h"
diff --git a/Engine/Source/Editor/LevelEditor/LevelEditor.Build.cs b/Engine/Source/Editor/LevelEditor/LevelEditor.Build.cs
--- a/Engine/Source/Editor/LevelEditor/LevelEditor.Build.cs
+++ b/Engine/Source/Editor/LevelEditor/LevelEditor.Build.cs
@@ -9,6 +9,7 @@ public class LevelEditor : ModuleRules
 		PrivateIncludePathModuleNames.AddRange(
 			new string[] {
 				"AssetTools",
+				"ClassViewer",
 				"MainFrame",
                 "PlacementMode",
 				"ReferenceViewer",
diff --git a/Engine/Source/Editor/LevelEditor/Private/LevelEditorToolBar.cpp b/Engine/Source/Editor/LevelEditor/Private/LevelEditorToolBar.cpp
--- a/Engine/Source/Editor/LevelEditor/Private/LevelEditorToolBar.cpp
+++ b/Engine/Source/Editor/LevelEditor/Private/LevelEditorToolBar.cpp
@@ -21,8 +21,8 @@
 #include "ISettingsModule.h"
 #include "ISettingsSection.h"
 #include "DesktopPlatformModule.h"
-#include "Editor/ClassViewer/Public/ClassViewerModule.h"
-#include "Editor/ClassViewer/Public/ClassViewerFilter.h"
+#include "ClassViewerModule.h"
+#include "ClassViewerFilter.h"
 #include "KismetEditorUtilities.h"
 #include "ISourceControlModule.h"
 #include "SVolumeControl.h"
diff --git a/Engine/Source/Editor/ProjectTargetPlatformEditor/ProjectTargetPlatformEditor.Build.cs b/Engine/Source/Editor/ProjectTargetPlatformEditor/ProjectTargetPlatformEditor.Build.cs
--- a/Engine/Source/Editor/ProjectTargetPlatformEditor/ProjectTargetPlatformEditor.Build.cs
+++ b/Engine/Source/Editor/ProjectTargetPlatformEditor/ProjectTargetPlatformEditor.Build.cs
@@ -22,7 +22,6 @@ public class ProjectTargetPlatformEditor : ModuleRules
 				"InputCore",
 				"EditorStyle",
 				"TargetPlatform",
-				"GameProjectGeneration",
 				"DesktopPlatform",
 				"Settings",
 				"UnrealEd",
@@ -30,6 +29,18 @@ public class ProjectTargetPlatformEditor : ModuleRules
 			}
 		);
 
+		PrivateIncludePathModuleNames.AddRange(
+			new string[] {
+				"GameProjectGeneration",
+			}
+		);
+
+		DynamicallyLoadedModuleNames.AddRange(
+			new string[] {
+				"GameProjectGeneration",
+			}
+		);
+
 		PrivateIncludePaths.AddRange(
 			new string[]
 			{
diff --git a/Engine/Source/Editor/PropertyEditor/Private/UserInterface/PropertyEditor/SPropertyEditorClass.cpp b/Engine/Source/Editor/PropertyEditor/Private/UserInterface/PropertyEditor/SPropertyEditorClass.cpp
--- a/Engine/Source/Editor/PropertyEditor/Private/UserInterface/PropertyEditor/SPropertyEditorClass.cpp
+++ b/Engine/Source/Editor/PropertyEditor/Private/UserInterface/PropertyEditor/SPropertyEditorClass.cpp
@@ -8,8 +8,8 @@
 #include "PropertyEditor.h"
 #include "SPropertyEditorClass.h"
 #include "SPropertyEditorCombo.h"
-#include "Editor/ClassViewer/Public/ClassViewerModule.h"
-#include "Editor/ClassViewer/Public/ClassViewerFilter.h"
+#include "ClassViewerModule.h"
+#include "ClassViewerFilter.h"
 
 #include "PropertyHandle.h"
 
diff --git a/Engine/Source/Editor/PropertyEditor/PropertyEditor.Build.cs b/Engine/Source/Editor/PropertyEditor/PropertyEditor.Build.cs
--- a/Engine/Source/Editor/PropertyEditor/PropertyEditor.Build.cs
+++ b/Engine/Source/Editor/PropertyEditor/PropertyEditor.Build.cs
@@ -39,6 +39,7 @@ public class PropertyEditor : ModuleRules
 				"MainFrame",
                 "AssetRegistry",
                 "AssetTools",
+				"ClassViewer",
                 "ContentBrowser"
 			}
 		);
@@ -63,6 +64,7 @@ public class PropertyEditor : ModuleRules
 			new string[] {
                 "AssetRegistry",
                 "AssetTools",
+				"ClassViewer",
 				"ContentBrowser",
 				"Documentation",
 				"MainFrame",
diff --git a/Engine/Source/Editor/UnrealEd/Private/Kismet2/BlueprintEditorUtils.cpp b/Engine/Source/Editor/UnrealEd/Private/Kismet2/BlueprintEditorUtils.cpp
--- a/Engine/Source/Editor/UnrealEd/Private/Kismet2/BlueprintEditorUtils.cpp
+++ b/Engine/Source/Editor/UnrealEd/Private/Kismet2/BlueprintEditorUtils.cpp
@@ -25,8 +25,8 @@
 #include "Editor/UnrealEd/Public/Kismet2/BlueprintEditorUtils.h"
 #include "Editor/UnrealEd/Public/Kismet2/StructureEditorUtils.h"
 #include "ScopedTransaction.h"
-#include "Editor/ClassViewer/Public/ClassViewerModule.h"
-#include "Editor/ClassViewer/Public/ClassViewerFilter.h"
+#include "ClassViewerModule.h"
+#include "ClassViewerFilter.h"
 #include "Editor/Kismet/Public/FindInBlueprintManager.h"
 
 #include "BlueprintEditor.h"
