GitHub pull request #706

diff --git a/Engine/Source/Editor/GameProjectGeneration/Private/GameProjectUtils.cpp b/Engine/Source/Editor/GameProjectGeneration/Private/GameProjectUtils.cpp
--- a/Engine/Source/Editor/GameProjectGeneration/Private/GameProjectUtils.cpp
+++ b/Engine/Source/Editor/GameProjectGeneration/Private/GameProjectUtils.cpp
@@ -3084,8 +3084,7 @@ bool GameProjectUtils::AddCodeToProject_Internal(const FString& NewClassName, co
 bool GameProjectUtils::FindSourceFileInProject(const FString& InFilename, const FString& InSearchPath, FString& OutPath)
 {
 	TArray<FString> Filenames;
-	const FString FilenameWidcard = TEXT("*") + InFilename;
-	IFileManager::Get().FindFilesRecursive(Filenames, *InSearchPath, *FilenameWidcard, true, false, false);
+	IFileManager::Get().FindFilesRecursive(Filenames, *InSearchPath, *InFilename, true, false, false);
 	
 	if(Filenames.Num())
 	{
