Added actor and scene components to the featured classes list when creating a C++ class in the CB

UE-9291 - Actor/SceneComponent are 'common' options when making BP, but not C++ class

diff --git a/Engine/Source/Editor/GameProjectGeneration/Public/AddToProjectConfig.h b/Engine/Source/Editor/GameProjectGeneration/Public/AddToProjectConfig.h
--- a/Engine/Source/Editor/GameProjectGeneration/Public/AddToProjectConfig.h
+++ b/Engine/Source/Editor/GameProjectGeneration/Public/AddToProjectConfig.h
@@ -107,14 +107,28 @@ struct FNewClassInfo
 /** @todo make this ini configurable */
 struct FFeaturedClasses
 {
+public:
 	/** Get a list of all featured native class types */
 	static TArray<FNewClassInfo> AllNativeClasses();
 	
 	/** Get a list of all featured Actor class types */
 	static TArray<FNewClassInfo> ActorClasses();
 
-	/** Get a list of all featured class types */
+	/** Get a list of all featured Component class types */
 	static TArray<FNewClassInfo> ComponentClasses();
+
+private:
+	/** Append the featured Actor class types that are commonly used */
+	static void AddCommonActorClasses(TArray<FNewClassInfo>& Array);
+
+	/** Append the featured Actor class types that are less commonly used */
+	static void AddExtraActorClasses(TArray<FNewClassInfo>& Array);
+
+	/** Append the featured Component class types that are commonly used */
+	static void AddCommonComponentClasses(TArray<FNewClassInfo>& Array);
+
+	/** Append the featured Component class types that are less commonly used */
+	static void AddExtraComponentClasses(TArray<FNewClassInfo>& Array);
 };
 
 /** Add to project dialog configuration structure */
diff --git a/Engine/Source/Editor/GameProjectGeneration/Public/FeaturedClasses.inl b/Engine/Source/Editor/GameProjectGeneration/Public/FeaturedClasses.inl
--- a/Engine/Source/Editor/GameProjectGeneration/Public/FeaturedClasses.inl
+++ b/Engine/Source/Editor/GameProjectGeneration/Public/FeaturedClasses.inl
@@ -27,7 +27,10 @@ TArray<FNewClassInfo> FFeaturedClasses::AllNativeClasses()
 	TArray<FNewClassInfo> Array;
 	Array.Add(FNewClassInfo(FNewClassInfo::EClassType::EmptyCpp));
 
-	Array += ActorClasses();
+	AddCommonActorClasses(Array);
+	AddCommonComponentClasses(Array);
+
+	AddExtraActorClasses(Array);
 
 	Array.Add(FNewClassInfo(UBlueprintFunctionLibrary::StaticClass()));
 
@@ -41,9 +44,31 @@ TArray<FNewClassInfo> FFeaturedClasses::AllNativeClasses()
 TArray<FNewClassInfo> FFeaturedClasses::ActorClasses()
 {
 	TArray<FNewClassInfo> Array;
+	AddCommonActorClasses(Array);
+	AddExtraActorClasses(Array);
+	return Array;
+}
+
+/** Get a list of all featured Component class types */
+TArray<FNewClassInfo> FFeaturedClasses::ComponentClasses()
+{
+	TArray<FNewClassInfo> Array;
+	AddCommonComponentClasses(Array);
+	AddExtraComponentClasses(Array);
+	return Array;
+}
+
+/** Append the featured Actor class types that are commonly used */
+void FFeaturedClasses::AddCommonActorClasses(TArray<FNewClassInfo>& Array)
+{
 	Array.Add(FNewClassInfo(ACharacter::StaticClass()));
 	Array.Add(FNewClassInfo(APawn::StaticClass()));
 	Array.Add(FNewClassInfo(AActor::StaticClass()));
+}
+
+/** Append the featured Actor class types that are less commonly used */
+void FFeaturedClasses::AddExtraActorClasses(TArray<FNewClassInfo>& Array)
+{
 	Array.Add(FNewClassInfo(APlayerCameraManager::StaticClass()));
 	Array.Add(FNewClassInfo(APlayerController::StaticClass()));
 	Array.Add(FNewClassInfo(AGameMode::StaticClass()));
@@ -51,18 +76,20 @@ TArray<FNewClassInfo> FFeaturedClasses::ActorClasses()
 	Array.Add(FNewClassInfo(AHUD::StaticClass()));
 	Array.Add(FNewClassInfo(APlayerState::StaticClass()));
 	Array.Add(FNewClassInfo(AGameState::StaticClass()));
-	return Array;
 }
 
-/** Get a list of all featured class types */
-TArray<FNewClassInfo> FFeaturedClasses::ComponentClasses()
+/** Append the featured Component class types that are commonly used */
+void FFeaturedClasses::AddCommonComponentClasses(TArray<FNewClassInfo>& Array)
 {
-	TArray<FNewClassInfo> Array;
 	Array.Add(FNewClassInfo(UActorComponent::StaticClass()));
 	Array.Add(FNewClassInfo(USceneComponent::StaticClass()));
+}
+
+/** Append the featured Component class types that are less commonly used */
+void FFeaturedClasses::AddExtraComponentClasses(TArray<FNewClassInfo>& Array)
+{
 	Array.Add(FNewClassInfo(UStaticMeshComponent::StaticClass()));
 	Array.Add(FNewClassInfo(USkeletalMeshComponent::StaticClass()));
 	Array.Add(FNewClassInfo(UCameraComponent::StaticClass()));
 	Array.Add(FNewClassInfo(UDirectionalLightComponent::StaticClass()));
-	return Array;
-}
\ No newline at end of file
+}
