It is no longer an error breaking the ability to create a new project if a new project files cannot be added to source control.  The attempt is still made however.

diff --git a/Engine/Source/Editor/GameProjectGeneration/Private/GameProjectUtils.cpp b/Engine/Source/Editor/GameProjectGeneration/Private/GameProjectUtils.cpp
--- a/Engine/Source/Editor/GameProjectGeneration/Private/GameProjectUtils.cpp
+++ b/Engine/Source/Editor/GameProjectGeneration/Private/GameProjectUtils.cpp
@@ -2581,10 +2581,7 @@ bool GameProjectUtils::GenerateGameResourceFiles(const FString& NewResourceFolde
 	FString IconFileName = GameRoot / TEXT("Build/Windows/Application.ico");
 	if(!FPaths::FileExists(IconFileName))
 	{
-		if(!SourceControlHelpers::CopyFileUnderSourceControl(IconFileName, FPaths::EngineContentDir() / TEXT("Editor/Templates/Resources/Windows/_GAME_NAME_.ico"), LOCTEXT("IconFileDescription", "icon"), OutFailReason))
-		{
-			return false;
-		}
+		SourceControlHelpers::CopyFileUnderSourceControl(IconFileName, FPaths::EngineContentDir() / TEXT("Editor/Templates/Resources/Windows/_GAME_NAME_.ico"), LOCTEXT("IconFileDescription", "icon"), OutFailReason);
 		OutCreatedFiles.Add(IconFileName);
 	}
 	
@@ -2618,10 +2615,7 @@ bool GameProjectUtils::GenerateGameResourceFiles(const FString& NewResourceFolde
 			}
 		};
 
-		if(!SourceControlHelpers::CheckoutOrMarkForAdd(OutputFilename, LOCTEXT("ResourceFileDescription", "resource"), FOnPostCheckOut::CreateStatic(&Local::WriteFile, &TemplateText, &OutCreatedFiles), OutFailReason))
-		{
-			return false;
-		}
+		SourceControlHelpers::CheckoutOrMarkForAdd(OutputFilename, LOCTEXT("ResourceFileDescription", "resource"), FOnPostCheckOut::CreateStatic(&Local::WriteFile, &TemplateText, &OutCreatedFiles), OutFailReason);
 	}
 #elif PLATFORM_MAC
 	//@todo MAC: Implement MAC version of these files...
