Fixed map path for starter content

diff --git a/Engine/Source/Editor/GameProjectGeneration/Private/GameProjectUtils.cpp b/Engine/Source/Editor/GameProjectGeneration/Private/GameProjectUtils.cpp
--- a/Engine/Source/Editor/GameProjectGeneration/Private/GameProjectUtils.cpp
+++ b/Engine/Source/Editor/GameProjectGeneration/Private/GameProjectUtils.cpp
@@ -1646,8 +1646,8 @@ bool GameProjectUtils::GenerateConfigFiles(const FProjectInformation& InProjectI
 				}
 				else
 				{
-					SpecificEditorStartupMap = TEXT("/Game/StarterContent/StarterMap");
-					SpecificGameDefaultMap = TEXT("/Game/StarterContent/StarterMap");
+					SpecificEditorStartupMap = TEXT("/Game/StarterContent/Maps/StarterMap");
+					SpecificGameDefaultMap = TEXT("/Game/StarterContent/Maps/StarterMap");
 				}
 			}						
 			
