[UE-11407] Fixing default map names for blank projects in Main to match 4.7.

diff --git a/Engine/Source/Editor/GameProjectGeneration/Private/GameProjectUtils.cpp b/Engine/Source/Editor/GameProjectGeneration/Private/GameProjectUtils.cpp
--- a/Engine/Source/Editor/GameProjectGeneration/Private/GameProjectUtils.cpp
+++ b/Engine/Source/Editor/GameProjectGeneration/Private/GameProjectUtils.cpp
@@ -1644,15 +1644,15 @@ bool GameProjectUtils::GenerateConfigFiles(const FProjectInformation& InProjectI
 			{
 				if (InProjectInfo.TargetedHardware == EHardwareClass::Mobile)
 				{
-					SpecificEditorStartupMap = TEXT("/Game/MobileStarterContent/Maps/StarterMap");
-					SpecificGameDefaultMap = TEXT("/Game/MobileStarterContent/Maps/StarterMap");
+					SpecificEditorStartupMap = TEXT("/Game/MobileStarterContent/Maps/Minimal_Default");
+					SpecificGameDefaultMap = TEXT("/Game/MobileStarterContent/Maps/Minimal_Default");
 				}
 				else
 				{
-					SpecificEditorStartupMap = TEXT("/Game/StarterContent/Maps/StarterMap");
-					SpecificGameDefaultMap = TEXT("/Game/StarterContent/Maps/StarterMap");
+					SpecificEditorStartupMap = TEXT("/Game/StarterContent/Maps/Minimal_Default");
+					SpecificGameDefaultMap = TEXT("/Game/StarterContent/Maps/Minimal_Default");
 				}
-			}						
+			}
 			
 			// Write out the settings for startup map and game default map
 			FileContents += TEXT("[/Script/EngineSettings.GameMapsSettings]") LINE_TERMINATOR;
