We now offer to show the Output Log when adding a new C++ fails in the editor

diff --git a/Engine/Source/Editor/GameProjectGeneration/Private/SNewClassDialog.cpp b/Engine/Source/Editor/GameProjectGeneration/Private/SNewClassDialog.cpp
--- a/Engine/Source/Editor/GameProjectGeneration/Private/SNewClassDialog.cpp
+++ b/Engine/Source/Editor/GameProjectGeneration/Private/SNewClassDialog.cpp
@@ -1210,8 +1210,11 @@ void SNewClassDialog::FinishClicked()
 		{
 			// @todo show fail reason in error label
 			// Failed to add code
-			const FText Message = FText::Format( LOCTEXT("AddCodeFailed_Native", "Failed to add class {0}. {1}"), FText::FromString(NewClassName), FailReason );
-			FMessageDialog::Open(EAppMsgType::Ok, Message);
+			const FText Message = FText::Format( LOCTEXT("AddCodeFailed_Native", "Failed to add or compile class {0}. {1}\n\nWould you like to open the Output Log to see more details?"), FText::FromString(NewClassName), FailReason );
+			if( FMessageDialog::Open(EAppMsgType::YesNo, Message) == EAppReturnType::Yes )
+			{
+				FGlobalTabmanager::Get()->InvokeTab(FName("OutputLog"));
+			}
 		}
 	}
 }
