New class workflow improvements

We now sync the Content Browser view after adding a new class, and also open the BP editor when adding a new BP class.

diff --git a/Engine/Source/Editor/GameProjectGeneration/GameProjectGeneration.Build.cs b/Engine/Source/Editor/GameProjectGeneration/GameProjectGeneration.Build.cs
--- a/Engine/Source/Editor/GameProjectGeneration/GameProjectGeneration.Build.cs
+++ b/Engine/Source/Editor/GameProjectGeneration/GameProjectGeneration.Build.cs
@@ -8,6 +8,8 @@ public class GameProjectGeneration : ModuleRules
 	{
         PrivateIncludePathModuleNames.AddRange(
             new string[] {
+				"AssetRegistry",
+				"ContentBrowser",
                 "DesktopPlatform",
                 "MainFrame",
             }
@@ -33,12 +35,13 @@ public class GameProjectGeneration : ModuleRules
 				"UnrealEd",
 				"DesktopPlatform",
                 "HardwareTargeting",
-                "AssetRegistry"
 			}
 		);
 
         DynamicallyLoadedModuleNames.AddRange(
             new string[] {
+				"AssetRegistry",
+				"ContentBrowser",
                 "DesktopPlatform",
                 "Documentation",
                 "MainFrame",
diff --git a/Engine/Source/Editor/GameProjectGeneration/Private/SNewClassDialog.cpp b/Engine/Source/Editor/GameProjectGeneration/Private/SNewClassDialog.cpp
--- a/Engine/Source/Editor/GameProjectGeneration/Private/SNewClassDialog.cpp
+++ b/Engine/Source/Editor/GameProjectGeneration/Private/SNewClassDialog.cpp
@@ -26,6 +26,8 @@
 #include "Kismet/BlueprintFunctionLibrary.h"
 #include "Kismet2/KismetEditorUtilities.h"
 #include "AssetRegistryModule.h"
+#include "AssetEditorManager.h"
+#include "ContentBrowserModule.h"
 #include "SNotificationList.h"
 #include "NotificationManager.h"
 
@@ -1125,6 +1127,15 @@ void SNewClassDialog::FinishClicked()
 
 					OnAddedToProject.ExecuteIfBound( NewClassName, PackagePath, FString() );
 
+					// Sync the content browser to the new asset
+					TArray<UObject*> SyncAssets;
+					SyncAssets.Add(NewBP);
+					FContentBrowserModule& ContentBrowserModule = FModuleManager::LoadModuleChecked<FContentBrowserModule>("ContentBrowser");
+					ContentBrowserModule.Get().SyncBrowserToAssets(SyncAssets);
+
+					// Open the editor for the new asset
+					FAssetEditorManager::Get().OpenEditorForAsset(NewBP);
+
 					// Successfully created the code and potentially opened the IDE. Close the dialog.
 					CloseContainingWindow();
 
@@ -1203,6 +1214,20 @@ void SNewClassDialog::FinishClicked()
 				}
 			}
 
+			// Sync the content browser to the new class
+			UPackage* const ClassPackage = FindPackage(nullptr, *(FString("/Script/") + SelectedModuleInfo->ModuleName));
+			if ( ClassPackage )
+			{
+				UClass* const NewClass = static_cast<UClass*>(FindObjectWithOuter(ClassPackage, UClass::StaticClass(), *NewClassName));
+				if ( NewClass )
+				{
+					TArray<UObject*> SyncAssets;
+					SyncAssets.Add(NewClass);
+					FContentBrowserModule& ContentBrowserModule = FModuleManager::LoadModuleChecked<FContentBrowserModule>("ContentBrowser");
+					ContentBrowserModule.Get().SyncBrowserToAssets(SyncAssets);
+				}
+			}
+
 			// Successfully created the code and potentially opened the IDE. Close the dialog.
 			CloseContainingWindow();
 		}
