fix for UE-10324 Unable to toggle between Public and Private class when using Add Code to Project

diff --git a/Engine/Source/Editor/GameProjectGeneration/Private/SNewClassDialog.cpp b/Engine/Source/Editor/GameProjectGeneration/Private/SNewClassDialog.cpp
--- a/Engine/Source/Editor/GameProjectGeneration/Private/SNewClassDialog.cpp
+++ b/Engine/Source/Editor/GameProjectGeneration/Private/SNewClassDialog.cpp
@@ -604,10 +604,9 @@ void SNewClassDialog::Construct( const FArguments& InArgs )
 												[
 													SNew(SCheckBox)
 													.Style(FEditorStyle::Get(), "Property.ToggleButton.Start")
-													.IsEnabled(this, &SNewClassDialog::CanChangeClassLocation)
 													.IsChecked(this, &SNewClassDialog::IsClassLocationActive, GameProjectUtils::EClassLocation::Public)
 													.OnCheckStateChanged(this, &SNewClassDialog::OnClassLocationChanged, GameProjectUtils::EClassLocation::Public)
-													.ToolTipText(this, &SNewClassDialog::GetClassLocationTooltip, GameProjectUtils::EClassLocation::Public)
+													.ToolTipText(LOCTEXT("ClassLocation_Public", "A public class can be included and used inside other modules in addition to the module it resides in"))
 													[
 														SNew(SBox)
 														.VAlign(VAlign_Center)
@@ -626,10 +625,9 @@ void SNewClassDialog::Construct( const FArguments& InArgs )
 												[
 													SNew(SCheckBox)
 													.Style(FEditorStyle::Get(), "Property.ToggleButton.End")
-													.IsEnabled(this, &SNewClassDialog::CanChangeClassLocation)
 													.IsChecked(this, &SNewClassDialog::IsClassLocationActive, GameProjectUtils::EClassLocation::Private)
 													.OnCheckStateChanged(this, &SNewClassDialog::OnClassLocationChanged, GameProjectUtils::EClassLocation::Private)
-													.ToolTipText(this, &SNewClassDialog::GetClassLocationTooltip, GameProjectUtils::EClassLocation::Private)
+													.ToolTipText(LOCTEXT("ClassLocation_Private", "A private class can only be included and used within the module it resides in"))
 													[
 														SNew(SBox)
 														.VAlign(VAlign_Center)
@@ -1327,26 +1325,6 @@ FSlateColor SNewClassDialog::GetClassLocationTextColor(GameProjectUtils::EClassL
 	return (ClassLocation == InLocation) ? FSlateColor(FLinearColor(0, 0, 0)) : FSlateColor(FLinearColor(0.72f, 0.72f, 0.72f, 1.f));
 }
 
-FText SNewClassDialog::GetClassLocationTooltip(GameProjectUtils::EClassLocation InLocation) const
-{
-	if(CanChangeClassLocation())
-	{
-		switch(InLocation)
-		{
-		case GameProjectUtils::EClassLocation::Public:
-			return LOCTEXT("ClassLocation_Public", "A public class can be included and used inside other modules in addition to the module it resides in");
-
-		case GameProjectUtils::EClassLocation::Private:
-			return LOCTEXT("ClassLocation_Private", "A private class can only be included and used within the module it resides in");
-
-		default:
-			break;
-		}
-	}
-
-	return LOCTEXT("ClassLocation_UserDefined", "Your project is either not using a Public and Private source layout, or you're explicitly creating your class outside of the Public or Private folder");
-}
-
 ECheckBoxState SNewClassDialog::IsClassLocationActive(GameProjectUtils::EClassLocation InLocation) const
 {
 	return (ClassLocation == InLocation) ? ECheckBoxState::Checked : ECheckBoxState::Unchecked;
@@ -1377,6 +1355,10 @@ void SNewClassDialog::OnClassLocationChanged(ECheckBoxState InCheckedState, Game
 			{
 				NewClassPath = AbsoluteClassPath.Replace(*RootPath, *PublicPath);
 			}
+			else
+			{
+				NewClassPath = PublicPath;
+			}
 			break;
 
 		case GameProjectUtils::EClassLocation::Private:
@@ -1388,6 +1370,10 @@ void SNewClassDialog::OnClassLocationChanged(ECheckBoxState InCheckedState, Game
 			{
 				NewClassPath = AbsoluteClassPath.Replace(*RootPath, *PrivatePath);
 			}
+			else
+			{
+				NewClassPath = PrivatePath;
+			}
 			break;
 
 		default:
@@ -1399,11 +1385,6 @@ void SNewClassDialog::OnClassLocationChanged(ECheckBoxState InCheckedState, Game
 	}
 }
 
-bool SNewClassDialog::CanChangeClassLocation() const
-{
-	return ClassLocation != GameProjectUtils::EClassLocation::UserDefined;
-}
-
 void SNewClassDialog::UpdateInputValidity()
 {
 	bLastInputValidityCheckSuccessful = true;
diff --git a/Engine/Source/Editor/GameProjectGeneration/Private/SNewClassDialog.h b/Engine/Source/Editor/GameProjectGeneration/Private/SNewClassDialog.h
--- a/Engine/Source/Editor/GameProjectGeneration/Private/SNewClassDialog.h
+++ b/Engine/Source/Editor/GameProjectGeneration/Private/SNewClassDialog.h
@@ -157,18 +157,12 @@ private:
 	/** Get the text color to use for the given class location checkbox */
 	FSlateColor GetClassLocationTextColor(GameProjectUtils::EClassLocation InLocation) const;
 
-	/** Get the tooltip to use for the given class location */
-	FText GetClassLocationTooltip(GameProjectUtils::EClassLocation InLocation) const;
-
 	/** Checks to see if the given class location is active based on the current value of NewClassPath */
 	ECheckBoxState IsClassLocationActive(GameProjectUtils::EClassLocation InLocation) const;
 
 	/** Update the value of NewClassPath so that it uses the given class location */
 	void OnClassLocationChanged(ECheckBoxState InCheckedState, GameProjectUtils::EClassLocation InLocation);
 
-	/** Checks to see if the class location can be changed (relies on us having a Public/Private folder layout) */
-	bool CanChangeClassLocation() const;
-
 	/** Checks the current class name/path for validity and updates cached values accordingly */
 	void UpdateInputValidity();
 
