UE-8578: Slate Widget fails to compile once added to a new project

Slate Widget was failing, because of missing Slate dependencies. Testing introduced a couple of problems which all was fixed by this CL:
1. I introduced AdditionalDependencies in .uproject file and change "Add Code To Project..." procedure to fill this array if needed. UBT reads this field and builds the project with required modules. Needed for Slate classes.
2. Changed UHT to #include missing headers in generated.h files if it was missing an include for it's super class. It was causing problems if we were trying to add a subclass of BrushShape -- BrushShape.h didn't have #include "Brush.h" and UBrushShape was inheriting from UBrush.
3. Above problems also occured for Slate classes, but not all of them was UCLASSes, so I had to fixed that manually.
4. "Add Code To Project..." functionality was not invalidating UBT makefiles, which lead to omitting new source files during hot-reloading (even thought it was reporting a success). This change also should improve a bit performance, cause right now there is no "gathering" step -- there is only invalidate step which is a lot quicker.
5. Fixed "Selected Class Source" link to source class in Slate Widget and Slate Widget Style class.

diff --git a/Engine/Source/Developer/DesktopPlatform/Private/DesktopPlatformBase.cpp b/Engine/Source/Developer/DesktopPlatform/Private/DesktopPlatformBase.cpp
--- a/Engine/Source/Developer/DesktopPlatform/Private/DesktopPlatformBase.cpp
+++ b/Engine/Source/Developer/DesktopPlatform/Private/DesktopPlatformBase.cpp
@@ -479,7 +479,7 @@ bool FDesktopPlatformBase::GenerateProjectFiles(const FString& RootDir, const FS
 	return bRes;
 }
 
-bool FDesktopPlatformBase::GatherProjectFiles(const FString& RootDir, const FString& ProjectFileName, FFeedbackContext* Warn)
+bool FDesktopPlatformBase::InvalidateMakefiles(const FString& RootDir, const FString& ProjectFileName, FFeedbackContext* Warn)
 {
 	// Composes the target, platform, and config (eg, "QAGame Win64 Development")
 	FString Arguments = FString::Printf(TEXT("%s %s %s"), FApp::GetGameName(), FPlatformMisc::GetUBTPlatform(), FModuleManager::GetUBTConfiguration());
@@ -498,10 +498,8 @@ bool FDesktopPlatformBase::GatherProjectFiles(const FString& RootDir, const FStr
 		}
 	}
 	
-	// -gatheronly tells UBT to update its UBT makefiles without building
-	Arguments += TEXT(" -gatheronly");
-
-	Arguments += TEXT(" -progress");
+	// -invalidatemakefilesonly tells UBT to invalidate its UBT makefiles without building
+	Arguments += TEXT(" -invalidatemakefilesonly");
 
 	if (FRocketSupport::IsRocket())
 	{
@@ -510,7 +508,7 @@ bool FDesktopPlatformBase::GatherProjectFiles(const FString& RootDir, const FStr
 
 	// Compile UnrealBuildTool if it doesn't exist. This can happen if we're just copying source from somewhere.
 	bool bRes = true;
-	Warn->BeginSlowTask(LOCTEXT("GatheringProjectFiles", "Gathering project files..."), true, true);
+	Warn->BeginSlowTask(LOCTEXT("InvalidateMakefiles", "Invalidating makefiles..."), true, true);
 	if(!FPaths::FileExists(GetUnrealBuildToolExecutableFilename(RootDir)))
 	{
 		Warn->StatusUpdate(0, 1, LOCTEXT("BuildingUBT", "Building UnrealBuildTool..."));
@@ -518,8 +516,8 @@ bool FDesktopPlatformBase::GatherProjectFiles(const FString& RootDir, const FStr
 	}
 	if(bRes)
 	{
-		Warn->StatusUpdate(0, 1, LOCTEXT("GatheringProjectFiles", "Gathering project files..."));
-		bRes = RunUnrealBuildTool(LOCTEXT("GatheringProjectFiles", "Gathering project files..."), RootDir, Arguments, Warn);
+		Warn->StatusUpdate(0, 1, LOCTEXT("InvalidateMakefiles", "Invalidating makefiles..."));
+		bRes = RunUnrealBuildTool(LOCTEXT("InvalidateMakefiles", "Invalidating makefiles..."), RootDir, Arguments, Warn);
 	}
 	Warn->EndSlowTask();
 	return bRes;
diff --git a/Engine/Source/Developer/DesktopPlatform/Private/DesktopPlatformBase.h b/Engine/Source/Developer/DesktopPlatform/Private/DesktopPlatformBase.h
--- a/Engine/Source/Developer/DesktopPlatform/Private/DesktopPlatformBase.h
+++ b/Engine/Source/Developer/DesktopPlatform/Private/DesktopPlatformBase.h
@@ -38,7 +38,7 @@ public:
 	virtual bool CleanGameProject(const FString& ProjectDir, FString& OutFailPath, FFeedbackContext* Warn) override;
 	virtual bool CompileGameProject(const FString& RootDir, const FString& ProjectFileName, FFeedbackContext* Warn) override;
 	virtual bool GenerateProjectFiles(const FString& RootDir, const FString& ProjectFileName, FFeedbackContext* Warn) override;
-	virtual bool GatherProjectFiles(const FString& RootDir, const FString& ProjectFileName, FFeedbackContext* Warn) override;
+	virtual bool InvalidateMakefiles(const FString& RootDir, const FString& ProjectFileName, FFeedbackContext* Warn) override;
 	virtual bool IsUnrealBuildToolAvailable() override;
 	virtual bool InvokeUnrealBuildToolSync(const FString& InCmdLineParams, FOutputDevice &Ar, bool bSkipBuildUBT, int32& OutReturnCode, FString& OutProcOutput) override;
 	virtual FProcHandle InvokeUnrealBuildToolAsync(const FString& InCmdLineParams, FOutputDevice &Ar, void*& OutReadPipe, void*& OutWritePipe, bool bSkipBuildUBT = false) override;
diff --git a/Engine/Source/Developer/DesktopPlatform/Public/IDesktopPlatform.h b/Engine/Source/Developer/DesktopPlatform/Public/IDesktopPlatform.h
--- a/Engine/Source/Developer/DesktopPlatform/Public/IDesktopPlatform.h
+++ b/Engine/Source/Developer/DesktopPlatform/Public/IDesktopPlatform.h
@@ -330,14 +330,14 @@ public:
 	virtual bool GenerateProjectFiles(const FString& RootDir, const FString& ProjectFileName, FFeedbackContext* Warn) = 0;
 
 	/**
-	* Gather project files for the given project (to update its UBT makefiles without performing a full project generation).
+	* Invalidate makefiles for project (to UBT regenerate them at startup).
 	*
 	* @param RootDir			Engine root directory for the project to use.
 	* @param ProjectFileName	Filename of the project to update
 	* @param Warn				Feedback context to use for progress updates
 	* @return true if project files were generated successfully.
 	*/
-	virtual bool GatherProjectFiles(const FString& RootDir, const FString& ProjectFileName, FFeedbackContext* Warn) = 0;
+	virtual bool InvalidateMakefiles(const FString& RootDir, const FString& ProjectFileName, FFeedbackContext* Warn) = 0;
 
 	/**
 	* Determines whether UnrealBuildTool is available
diff --git a/Engine/Source/Editor/GameProjectGeneration/Private/GameProjectUtils.cpp b/Engine/Source/Editor/GameProjectGeneration/Private/GameProjectUtils.cpp
--- a/Engine/Source/Editor/GameProjectGeneration/Private/GameProjectUtils.cpp
+++ b/Engine/Source/Editor/GameProjectGeneration/Private/GameProjectUtils.cpp
@@ -207,12 +207,48 @@ bool FNewClassInfo::GetIncludePath(FString& OutIncludePath) const
 		}
 		break;
 
+	case EClassType::SlateWidget:
+		OutIncludePath = "Widgets/SCompoundWidget.h";
+		return true;
+
+	case EClassType::SlateWidgetStyle:
+		OutIncludePath = "Styling/SlateWidgetStyle.h";
+		return true;
+
 	default:
 		break;
 	}
 	return false;
 }
 
+FString FNewClassInfo::GetBaseClassHeaderFilename() const
+{
+	FString IncludePath;
+
+	switch (ClassType)
+	{
+	case EClassType::UObject:
+		if (BaseClass)
+		{
+			FString ClassHeaderPath;
+			if (FSourceCodeNavigation::FindClassHeaderPath(BaseClass, ClassHeaderPath) && IFileManager::Get().FileSize(*ClassHeaderPath) != INDEX_NONE)
+			{
+				return ClassHeaderPath;
+			}
+		}
+		break;
+
+	case EClassType::SlateWidget:
+	case EClassType::SlateWidgetStyle:
+		GetIncludePath(IncludePath);
+		return FPaths::EngineDir() / TEXT("Source") / TEXT("Runtime") / TEXT("SlateCore") / TEXT("Public") / IncludePath;
+	default:
+		return FString();
+	}
+
+	return FString();
+}
+
 FString FNewClassInfo::GetHeaderFilename(const FString& ClassName) const
 {
 	const FString HeaderFilename = GetFinalClassName(ClassName) + TEXT(".h");
@@ -688,9 +724,50 @@ void GameProjectUtils::OnWarningReasonOk()
 	}
 }
 
+bool GameProjectUtils::UpdateStartupModuleNames(FProjectDescriptor& Descriptor, const TArray<FString>* StartupModuleNames)
+{
+	if (StartupModuleNames == nullptr)
+	{
+		return false;
+	}
+
+	// Replace the modules names, if specified
+	Descriptor.Modules.Empty();
+	for (int32 Idx = 0; Idx < StartupModuleNames->Num(); Idx++)
+	{
+		Descriptor.Modules.Add(FModuleDescriptor(*(*StartupModuleNames)[Idx]));
+	}
+
+	return true;
+}
+
+bool GameProjectUtils::UpdateRequiredAdditionalDependencies(FProjectDescriptor& Descriptor, TArray<FString>& RequiredDependencies, const FString& ModuleName)
+{
+	bool bNeedsUpdate = false;
+
+	for (auto& ModuleDesc : Descriptor.Modules)
+	{
+		if (ModuleDesc.Name != *ModuleName)
+		{
+			continue;
+		}
+
+		for (const auto& RequiredDep : RequiredDependencies)
+		{
+			if (!ModuleDesc.AdditionalDependencies.Contains(RequiredDep))
+			{
+				ModuleDesc.AdditionalDependencies.Add(RequiredDep);
+				bNeedsUpdate = true;
+			}
+		}
+	}
+
+	return bNeedsUpdate;
+}
+
 bool GameProjectUtils::UpdateGameProject(const FString& ProjectFile, const FString& EngineIdentifier, FText& OutFailReason)
 {
-	return UpdateGameProjectFile(ProjectFile, EngineIdentifier, NULL, OutFailReason);
+	return UpdateGameProjectFile(ProjectFile, EngineIdentifier, OutFailReason);
 }
 
 void GameProjectUtils::OpenAddToProjectDialog(const FAddToProjectConfig& Config, EClassDomain InDomain)
@@ -1702,7 +1779,12 @@ bool GameProjectUtils::GenerateBasicSourceCode(TArray<FString>& OutCreatedFiles,
 	TArray<FString> StartupModuleNames;
 	if (GameProjectUtils::GenerateBasicSourceCode(FPaths::GameSourceDir().LeftChop(1), FApp::GetGameName(), FPaths::GameDir(), StartupModuleNames, OutCreatedFiles, OutFailReason))
 	{
-		GameProjectUtils::UpdateProject(&StartupModuleNames);
+		GameProjectUtils::UpdateProject(
+			FProjectDescriptorModifier::CreateLambda(
+			[&StartupModuleNames](FProjectDescriptor& Descriptor)
+			{
+				return UpdateStartupModuleNames(Descriptor, &StartupModuleNames);
+			}));
 		return true;
 	}
 
@@ -2690,17 +2772,27 @@ bool GameProjectUtils::GenerateGameModuleHeaderFile(const FString& NewBuildFileN
 
 void GameProjectUtils::OnUpdateProjectConfirm()
 {
-	UpdateProject(NULL);
+	UpdateProject();
 }
 
-void GameProjectUtils::UpdateProject(const TArray<FString>* StartupModuleNames)
+void GameProjectUtils::UpdateProject(const FProjectDescriptorModifier& Modifier)
+{
+	UpdateProject_Impl(&Modifier);
+}
+
+void GameProjectUtils::UpdateProject()
+{
+	UpdateProject_Impl(nullptr);
+}
+
+void GameProjectUtils::UpdateProject_Impl(const FProjectDescriptorModifier* Modifier)
 {
 	const FString& ProjectFilename = FPaths::GetProjectFilePath();
 	const FString& ShortFilename = FPaths::GetCleanFilename(ProjectFilename);
 	FText FailReason;
 	FText UpdateMessage;
 	SNotificationItem::ECompletionState NewCompletionState;
-	if ( UpdateGameProjectFile(ProjectFilename, FDesktopPlatformModule::Get()->GetCurrentEngineIdentifier(), StartupModuleNames, FailReason) )
+	if (UpdateGameProjectFile_Impl(ProjectFilename, FDesktopPlatformModule::Get()->GetCurrentEngineIdentifier(), Modifier, FailReason))
 	{
 		// The project was updated successfully.
 		FFormatNamedArguments Args;
@@ -2727,6 +2819,21 @@ void GameProjectUtils::UpdateProject(const TArray<FString>* StartupModuleNames)
 	}
 }
 
+void GameProjectUtils::UpdateProject(const TArray<FString>* StartupModuleNames)
+{
+	UpdateProject(
+		FProjectDescriptorModifier::CreateLambda(
+			[StartupModuleNames](FProjectDescriptor& Desc)
+			{
+				if (StartupModuleNames != nullptr)
+				{
+					return UpdateStartupModuleNames(Desc, StartupModuleNames);
+				}
+
+				return false;
+			}));
+}
+
 void GameProjectUtils::OnUpdateProjectCancel()
 {
 	if ( UpdateGameProjectNotification.IsValid() )
@@ -2761,7 +2868,17 @@ void GameProjectUtils::TryMakeProjectFileWriteable(const FString& ProjectFile)
 	}
 }
 
-bool GameProjectUtils::UpdateGameProjectFile(const FString& ProjectFile, const FString& EngineIdentifier, const TArray<FString>* StartupModuleNames, FText& OutFailReason)
+bool GameProjectUtils::UpdateGameProjectFile(const FString& ProjectFile, const FString& EngineIdentifier, const FProjectDescriptorModifier& Modifier, FText& OutFailReason)
+{
+	return UpdateGameProjectFile_Impl(ProjectFile, EngineIdentifier, &Modifier, OutFailReason);
+}
+
+bool GameProjectUtils::UpdateGameProjectFile(const FString& ProjectFile, const FString& EngineIdentifier, FText& OutFailReason)
+{
+	return UpdateGameProjectFile_Impl(ProjectFile, EngineIdentifier, nullptr, OutFailReason);
+}
+
+bool GameProjectUtils::UpdateGameProjectFile_Impl(const FString& ProjectFile, const FString& EngineIdentifier, const FProjectDescriptorModifier* Modifier, FText& OutFailReason)
 {
 	// Make sure we can write to the project file
 	TryMakeProjectFileWriteable(ProjectFile);
@@ -2770,14 +2887,10 @@ bool GameProjectUtils::UpdateGameProjectFile(const FString& ProjectFile, const F
 	FProjectDescriptor Descriptor;
 	if(Descriptor.Load(ProjectFile, OutFailReason))
 	{
-		// Replace the modules names, if specified
-		if(StartupModuleNames != NULL)
+		if (Modifier && Modifier->IsBound() && !Modifier->Execute(Descriptor))
 		{
-			Descriptor.Modules.Empty();
-			for(int32 Idx = 0; Idx < StartupModuleNames->Num(); Idx++)
-			{
-				Descriptor.Modules.Add(FModuleDescriptor(*(*StartupModuleNames)[Idx]));
-			}
+			// If modifier returns false it means that we want to drop changes.
+			return true;
 		}
 
 		// Update file on disk
@@ -2786,6 +2899,22 @@ bool GameProjectUtils::UpdateGameProjectFile(const FString& ProjectFile, const F
 	return false;
 }
 
+bool GameProjectUtils::UpdateGameProjectFile(const FString& ProjectFilename, const FString& EngineIdentifier, const TArray<FString>* StartupModuleNames, FText& OutFailReason)
+{
+	return UpdateGameProjectFile(ProjectFilename, EngineIdentifier,
+		FProjectDescriptorModifier::CreateLambda(
+			[StartupModuleNames](FProjectDescriptor& Desc)
+			{
+				if (StartupModuleNames != nullptr)
+				{
+					return UpdateStartupModuleNames(Desc, StartupModuleNames);
+				}
+
+				return false;
+			}
+		), OutFailReason);
+}
+
 bool GameProjectUtils::CheckoutGameProjectFile(const FString& ProjectFilename, FText& OutFailReason)
 {
 	if ( !ensure(ProjectFilename.Len()) )
@@ -2869,6 +2998,31 @@ bool GameProjectUtils::ProjectHasCodeFiles()
 	return GameProjectUtils::GetProjectCodeFileCount() > 0;
 }
 
+TArray<FString> GameProjectUtils::GetRequiredAdditionalDependencies(const FNewClassInfo& ClassInfo)
+{
+	TArray<FString> Out;
+
+	switch (ClassInfo.ClassType)
+	{
+	case FNewClassInfo::EClassType::SlateWidget:
+	case FNewClassInfo::EClassType::SlateWidgetStyle:
+		Out.Reserve(2);
+		Out.Add(TEXT("Slate"));
+		Out.Add(TEXT("SlateCore"));
+		break;
+
+	case FNewClassInfo::EClassType::UObject:
+		auto ClassPackageName = ClassInfo.BaseClass->GetOutermost()->GetFName().ToString();
+
+		checkf(ClassPackageName.StartsWith(TEXT("/Script/")), TEXT("Class outermost should start with /Script/"));
+
+		Out.Add(ClassPackageName.Mid(8)); // Skip the /Script/ prefix.
+		break;
+	}
+
+	return Out;
+}
+
 bool GameProjectUtils::AddCodeToProject_Internal(const FString& NewClassName, const FString& NewClassPath, const FModuleContextInfo& ModuleInfo, const FNewClassInfo ParentClassInfo, const TSet<FString>& DisallowedHeaderNames, FString& OutHeaderFilePath, FString& OutCppFilePath, FText& OutFailReason)
 {
 	if ( !ParentClassInfo.IsSet() )
@@ -2902,9 +3056,17 @@ bool GameProjectUtils::AddCodeToProject_Internal(const FString& NewClassName, co
 	SlowTask.MakeDialog();
 
 	SlowTask.EnterProgressFrame();
+
+	auto RequiredDependencies = GetRequiredAdditionalDependencies(ParentClassInfo);
+	RequiredDependencies.Remove(ModuleInfo.ModuleName);
+
+	// Update project file if needed.
+	auto bUpdateProjectModules = false;
 	
 	// If the project does not already contain code, add the primary game module
 	TArray<FString> CreatedFiles;
+	TArray<FString> StartupModuleNames;
+
 	const bool bProjectHadCodeFiles = ProjectHasCodeFiles();
 	if (!bProjectHadCodeFiles)
 	{
@@ -2914,10 +3076,9 @@ bool GameProjectUtils::AddCodeToProject_Internal(const FString& NewClassName, co
 		// Assuming the game name is the same as the primary game module name
 		const FString GameModuleName = FApp::GetGameName();
 
-		TArray<FString> StartupModuleNames;
 		if ( GenerateBasicSourceCode(SourceDir, GameModuleName, FPaths::GameDir(), StartupModuleNames, CreatedFiles, OutFailReason) )
 		{
-			UpdateProject(&StartupModuleNames);
+			bUpdateProjectModules = true;
 		}
 		else
 		{
@@ -2926,6 +3087,21 @@ bool GameProjectUtils::AddCodeToProject_Internal(const FString& NewClassName, co
 		}
 	}
 
+	if (RequiredDependencies.Num() > 0 || bUpdateProjectModules)
+	{
+		UpdateProject(
+			FProjectDescriptorModifier::CreateLambda(
+			[&StartupModuleNames, &RequiredDependencies, &ModuleInfo, bUpdateProjectModules](FProjectDescriptor& Descriptor)
+			{
+				bool bNeedsUpdate = false;
+
+				bNeedsUpdate |= UpdateStartupModuleNames(Descriptor, bUpdateProjectModules ? &StartupModuleNames : nullptr);
+				bNeedsUpdate |= UpdateRequiredAdditionalDependencies(Descriptor, RequiredDependencies, ModuleInfo.ModuleName);
+
+				return bNeedsUpdate;
+			}));
+	}
+
 	SlowTask.EnterProgressFrame();
 
 	// Class Header File
@@ -2975,7 +3151,7 @@ bool GameProjectUtils::AddCodeToProject_Internal(const FString& NewClassName, co
 	if ( bProjectHadCodeFiles && FSourceCodeNavigation::AddSourceFiles(CreatedFilesForExternalAppRead) )
 	{
 		// We successfully added the new files to the solution, but we still need to run UBT with -gather to update any UBT makefiles
-		if ( FDesktopPlatformModule::Get()->GatherProjectFiles(FPaths::RootDir(), FPaths::GetProjectFilePath(), GWarn) )
+		if ( FDesktopPlatformModule::Get()->InvalidateMakefiles(FPaths::RootDir(), FPaths::GetProjectFilePath(), GWarn) )
 		{
 			// We managed the gather, so we can skip running the full generate
 			bGenerateProjectFiles = false;
diff --git a/Engine/Source/Editor/GameProjectGeneration/Private/GameProjectUtils.h b/Engine/Source/Editor/GameProjectGeneration/Private/GameProjectUtils.h
--- a/Engine/Source/Editor/GameProjectGeneration/Private/GameProjectUtils.h
+++ b/Engine/Source/Editor/GameProjectGeneration/Private/GameProjectUtils.h
@@ -29,6 +29,8 @@ struct FProjectInformation
 	EGraphicsPreset::Type DefaultGraphicsPerformance;
 };
 
+DECLARE_DELEGATE_RetVal_OneParam(bool, FProjectDescriptorModifier, FProjectDescriptor&);
+
 class GameProjectUtils
 {
 public:
@@ -309,6 +311,18 @@ private:
 	/** @param OutProjectCodeFilenames Contains the filenames of the project source code files */
 	static void GetProjectCodeFilenames(TArray<FString>& OutProjectCodeFilenames);
 
+	/**
+	 * Updates the projects and modifies FProjectDescriptor accordingly to given modifier.
+	 *
+	 * @param Modifier	Callback delegate that will modify the project descriptor accordingly.
+	 */
+	static void UpdateProject(const FProjectDescriptorModifier& Modifier);
+
+	/**
+	 * Updates the project file.
+	 */
+	static void UpdateProject();
+
 	/**
 	 * Updates the projects, and optionally the modules names
 	 *
@@ -331,6 +345,29 @@ private:
 	 */
 	static bool UpdateGameProjectFile(const FString& ProjectFilename, const FString& EngineIdentifier, const TArray<FString>* StartupModuleNames, FText& OutFailReason);
 
+	/**
+	 * Updates the loaded game project file to the current version and modifies FProjectDescriptor accordingly to given modifier.
+	 *
+	 * @param ProjectFilename		The name of the project (used to checkout from source control)
+	 * @param EngineIdentifier		The identifier for the engine to open the project with
+	 * @param Modifier				Callback delegate that will modify the project descriptor accordingly.
+	 * @param OutFailReason			Out, if unsuccessful this is the reason why
+
+	 * @return true, if successful
+	 */
+	static bool UpdateGameProjectFile(const FString& ProjectFilename, const FString& EngineIdentifier, const FProjectDescriptorModifier& Modifier, FText& OutFailReason);
+
+	/**
+	 * Updates the loaded game project file to the current version.
+	 *
+	 * @param ProjectFilename		The name of the project (used to checkout from source control)
+	 * @param EngineIdentifier		The identifier for the engine to open the project with
+	 * @param OutFailReason			Out, if unsuccessful this is the reason why
+
+	 * @return true, if successful
+	 */
+	static bool UpdateGameProjectFile(const FString& ProjectFilename, const FString& EngineIdentifier, FText& OutFailReason);
+
 	/** Checks the specified game project file out from source control */
 	static bool CheckoutGameProjectFile(const FString& ProjectFilename, FText& OutFailReason);
 
@@ -347,7 +384,56 @@ private:
 	/** Given a source file name, find its location within the project */
 	static bool FindSourceFileInProject(const FString& InFilename, const FString& InSearchPath, FString& OutPath);
 
+	/**
+	 * Gets required additional dependencies for fresh project for given class info.
+	 *
+	 * @param ClassInfo Given class info.
+	 *
+	 * @returns Array of required dependencies.
+	 */
+	static TArray<FString> GetRequiredAdditionalDependencies(const FNewClassInfo& ClassInfo);
+
+	/**
+	 * Updates startup module names in project descriptor.
+	 *
+	 * @param Descriptor Descriptor to update.
+	 * @param StartupModuleNames Modules to fill.
+	 *
+	 * @returns True if descriptor has been modified. False otherwise.
+	 */
+	static bool UpdateStartupModuleNames(FProjectDescriptor& Descriptor, const TArray<FString>* StartupModuleNames);
+
+	/**
+	 * Updates additional dependencies in project descriptor.
+	 *
+	 * @param Descriptor Descriptor to update.
+	 * @param RequiredDependencies Required dependencies.
+	 * @param ModuleName Module name for which those dependencies are required.
+	 *
+	 * @returns True if descriptor has been modified. False otherwise.
+	 */
+	static bool UpdateRequiredAdditionalDependencies(FProjectDescriptor& Descriptor, TArray<FString>& RequiredDependencies, const FString& ModuleName);
+
 private:
+	/**
+	 * Updates the projects and modifies FProjectDescriptor accordingly to given modifier.
+	 *
+	 * @param Modifier	Callback delegate that will modify the project descriptor accordingly.
+	 */
+	static void UpdateProject_Impl(const FProjectDescriptorModifier* Modifier);
+
+	/**
+	 * Updates the loaded game project file to the current version and modifies FProjectDescriptor accordingly to given modifier.
+	 *
+	 * @param ProjectFilename		The name of the project (used to checkout from source control)
+	 * @param EngineIdentifier		The identifier for the engine to open the project with
+	 * @param Modifier				Callback delegate that will modify the project descriptor accordingly.
+	 * @param OutFailReason			Out, if unsuccessful this is the reason why
+
+	 * @return true, if successful
+	 */
+	static bool UpdateGameProjectFile_Impl(const FString& ProjectFilename, const FString& EngineIdentifier, const FProjectDescriptorModifier* Modifier, FText& OutFailReason);
+
 	static TWeakPtr<SNotificationItem> UpdateGameProjectNotification;
 	static TWeakPtr<SNotificationItem> WarningProjectNameNotification;
 };
diff --git a/Engine/Source/Editor/GameProjectGeneration/Private/SNewClassDialog.cpp b/Engine/Source/Editor/GameProjectGeneration/Private/SNewClassDialog.cpp
--- a/Engine/Source/Editor/GameProjectGeneration/Private/SNewClassDialog.cpp
+++ b/Engine/Source/Editor/GameProjectGeneration/Private/SNewClassDialog.cpp
@@ -883,12 +883,12 @@ EVisibility SNewClassDialog::GetSourceHyperlinkVisibility() const
 		return EVisibility::Collapsed;
 	}
 
-	return (GetClassHeaderPath(ParentClassInfo.BaseClass).Len() > 0 ? EVisibility::Visible : EVisibility::Hidden);
+	return (ParentClassInfo.GetBaseClassHeaderFilename().Len() > 0 ? EVisibility::Visible : EVisibility::Hidden);
 }
 
 FText SNewClassDialog::GetSelectedParentClassFilename() const
 {
-	const FString ClassHeaderPath = GetClassHeaderPath(ParentClassInfo.BaseClass);
+	const FString ClassHeaderPath = ParentClassInfo.GetBaseClassHeaderFilename();
 	if (ClassHeaderPath.Len() > 0)
 	{
 		return FText::FromString(FPaths::GetCleanFilename(*ClassHeaderPath));
@@ -908,7 +908,7 @@ FString SNewClassDialog::GetSelectedParentDocLink() const
 
 void SNewClassDialog::OnEditCodeClicked()
 {
-	const FString ClassHeaderPath = GetClassHeaderPath(ParentClassInfo.BaseClass);
+	const FString ClassHeaderPath = ParentClassInfo.GetBaseClassHeaderFilename();
 	if (ClassHeaderPath.Len() > 0)
 	{
 		const FString AbsoluteHeaderPath = IFileManager::Get().ConvertToAbsolutePathForExternalAppForRead(*ClassHeaderPath);
diff --git a/Engine/Source/Editor/GameProjectGeneration/Public/AddToProjectConfig.h b/Engine/Source/Editor/GameProjectGeneration/Public/AddToProjectConfig.h
--- a/Engine/Source/Editor/GameProjectGeneration/Public/AddToProjectConfig.h
+++ b/Engine/Source/Editor/GameProjectGeneration/Public/AddToProjectConfig.h
@@ -84,6 +84,9 @@ struct FNewClassInfo
 	/** Get the path needed to include this class into another file */
 	bool GetIncludePath(FString& OutIncludePath) const;
 
+	/** Gets header filename of the base class. */
+	FString GetBaseClassHeaderFilename() const;
+
 	/** Given a class name, generate the header file (.h) that should be used for this class */
 	FString GetHeaderFilename(const FString& ClassName) const;
 
diff --git a/Engine/Source/Programs/UnrealBuildTool/System/RulesCompiler.cs b/Engine/Source/Programs/UnrealBuildTool/System/RulesCompiler.cs
--- a/Engine/Source/Programs/UnrealBuildTool/System/RulesCompiler.cs
+++ b/Engine/Source/Programs/UnrealBuildTool/System/RulesCompiler.cs
@@ -320,6 +320,46 @@ namespace UnrealBuildTool
 
 		/** Redistribution override flag for this module. */
 		public bool? IsRedistributableOverride { get; set; }
+
+		/**
+		 * Reads additional dependencies array for project module from project file and fills PrivateDependencyModuleNames. 
+		 *
+		 * @param ProjectFile A path to the .uproject file.
+		 * @param ModuleName Name of the module.
+		 */
+		public void ReadAdditionalDependencies(string ProjectFile, string ModuleName)
+		{
+			// Create a case-insensitive dictionary of the contents
+			Dictionary<string, object> Descriptor = fastJSON.JSON.Instance.ToObject<Dictionary<string, object>>(File.ReadAllText(ProjectFile));
+			Descriptor = new Dictionary<string,object>(Descriptor, StringComparer.InvariantCultureIgnoreCase);
+
+			// Get the list of plugins
+			object ModulesObject;
+			if (Descriptor.TryGetValue("Modules", out ModulesObject))
+			{
+				foreach(var ModuleObject in (ModulesObject as object[]).Cast<Dictionary<string, object>>())
+				{
+					object NameObject;
+					object AdditionalDependenciesObject;
+
+					if(!ModuleObject.TryGetValue("Name", out NameObject)
+						|| !(NameObject as string).Equals(ModuleName)
+						|| !ModuleObject.TryGetValue("AdditionalDependencies", out AdditionalDependenciesObject))
+					{
+						continue;
+					}
+
+					foreach (var AdditionalDependency in (AdditionalDependenciesObject as object[]).Cast<string>())
+					{
+						if(!PrivateDependencyModuleNames.Contains(AdditionalDependency))
+						{
+							PrivateDependencyModuleNames.Add(AdditionalDependency);
+						}
+					}
+					break;
+				}
+			}
+		}
 	}
 
 	/// <summary>
@@ -1384,6 +1424,19 @@ namespace UnrealBuildTool
 				throw new BuildException( Ex, "Unable to instantiate instance of '{0}' object type from compiled assembly '{1}'.  Unreal Build Tool creates an instance of your module's 'Rules' object in order to find out about your module's requirements.  The CLR exception details may provide more information:  {2}", ModuleTypeName, AssemblyFileName, Ex.ToString() );
 			}
 
+			// Have to do absolute here as this could be a project that is under the root
+			var FullUProjectPath = string.IsNullOrWhiteSpace(UnrealBuildTool.GetUProjectPath())
+				? ""
+				: Path.GetFullPath(UnrealBuildTool.GetUProjectPath());
+			var bProjectModule = string.IsNullOrWhiteSpace(FullUProjectPath)
+				? false
+				: Utils.IsFileUnderDirectory(ModuleFileName, FullUProjectPath);
+
+			if (bProjectModule)
+			{
+				RulesObject.ReadAdditionalDependencies(UnrealBuildTool.GetUProjectFile(), ModuleName);
+			}
+
 			// Validate rules object
 			{
 				if( RulesObject.Type == ModuleRules.ModuleType.CPlusPlus )
diff --git a/Engine/Source/Programs/UnrealBuildTool/System/UnrealBuildTool.cs b/Engine/Source/Programs/UnrealBuildTool/System/UnrealBuildTool.cs
--- a/Engine/Source/Programs/UnrealBuildTool/System/UnrealBuildTool.cs
+++ b/Engine/Source/Programs/UnrealBuildTool/System/UnrealBuildTool.cs
@@ -72,6 +72,9 @@ namespace UnrealBuildTool
         /** This is set to true during UBT startup, after it is safe to be accessing the IsGatheringBuild and IsAssemblingBuild properties.  Never access this directly. */
         private static bool bIsSafeToCheckIfGatheringOrAssemblingBuild = false;
 
+		/** True if this run of UBT should only invalidate Makefile. */
+		static public bool bIsInvalidatingMakefilesOnly = false;
+
         /** True if we should gather module dependencies for building in this run.  If this is false, then we'll expect to be loading information from disk about
             the target's modules before we'll be able to build anything.  One or both of IsGatheringBuild or IsAssemblingBuild must be true. */
         static public bool IsGatheringBuild
@@ -948,6 +951,10 @@ namespace UnrealBuildTool
                             BuildConfiguration.bXGEExport = true;
                             BuildConfiguration.bAllowXGE = true;
                         }
+						else if (LowercaseArg == "-invalidatemakefilesonly")
+						{
+							UnrealBuildTool.bIsInvalidatingMakefilesOnly = true;
+						}
                         else if (LowercaseArg == "-gather")
                         {
                             UnrealBuildTool.bIsGatheringBuild_Unsafe = true;
@@ -1426,6 +1433,21 @@ namespace UnrealBuildTool
 					}
 				}
 
+				if (UnrealBuildTool.bIsInvalidatingMakefilesOnly)
+				{
+					Log.TraceInformation("Invalidating makefiles only in this run.");
+					if (TargetDescs != null && TargetDescs.Count == 1)
+					{
+						InvalidateMakefiles(TargetDescs[0]);
+						return ECompilationResult.Succeeded;
+					}
+					else
+					{
+						Log.TraceError("You have to provide one target name for makefile invalidation.");
+						return ECompilationResult.OtherCompilationError;
+					}
+				}
+
 				UEBuildConfiguration.bHotReloadFromIDE = UEBuildConfiguration.bAllowHotReloadFromIDE && TargetDescs.Count == 1 && !TargetDescs[0].bIsEditorRecompile && ShouldDoHotReloadFromIDE(TargetDescs[0]);
 				bool bIsHotReload = UEBuildConfiguration.bHotReloadFromIDE || ( TargetDescs.Count == 1 && TargetDescs[0].OnlyModules.Count > 0 );
 				TargetDescriptor HotReloadTargetDesc = bIsHotReload ? TargetDescs[0] : null;
@@ -1887,6 +1909,25 @@ namespace UnrealBuildTool
             return BuildResult;
         }
 
+		/// <summary>
+		/// Invalidates makefiles for given target.
+		/// </summary>
+		/// <param name="Target">Target</param>
+		private static void InvalidateMakefiles(TargetDescriptor Target)
+		{
+			var MakefileNames = new string[] { "HotReloadMakefile.ubt", "Makefile.ubt" };
+			var BaseDir = GetUBTMakefileDirectoryPathForSingleTarget(Target);
+
+			foreach (var MakefileName in MakefileNames)
+			{
+				var MakefilePath = Path.Combine(BaseDir, MakefileName);
+
+				if (File.Exists(MakefilePath))
+				{
+					File.Delete(MakefilePath);
+				}
+			}
+		}
 		/// <summary>
 		/// Checks if the editor is currently running and this is a hot-reload
 		/// </summary>
@@ -2340,28 +2381,10 @@ namespace UnrealBuildTool
             {
 				var TargetDesc = TargetDescs[0];
 
-                // If there's only one target, just save the UBTMakefile in the target's build intermediate directory
-                // under a folder for that target (and platform/config combo.)
-                string PlatformIntermediatePath = BuildConfiguration.PlatformIntermediatePath;
-                if (UnrealBuildTool.HasUProjectFile())
-                {
-                    PlatformIntermediatePath = Path.Combine(UnrealBuildTool.GetUProjectPath(), BuildConfiguration.PlatformIntermediateFolder);
-                }
-
-				bool bIsHotReload = ( UEBuildConfiguration.bHotReloadFromIDE || ( TargetDesc.OnlyModules != null && TargetDesc.OnlyModules.Count > 0 ) );
+				bool bIsHotReload = (UEBuildConfiguration.bHotReloadFromIDE || (TargetDesc.OnlyModules != null && TargetDesc.OnlyModules.Count > 0));
 				string UBTMakefileName = bIsHotReload ? "HotReloadMakefile.ubt" : "Makefile.ubt";
 
-				// @todo ubtmake: If this is a compile triggered from the editor it will have passed along the game's target name, not the editor target name.
-				// At this point in Unreal Build Tool, we can't possibly know what the actual editor target name is, but we can take a guess.
-				// Even if we get it wrong, this won't have any side effects aside from not being able to share the exact same cached UBT Makefile
-				// between hot reloads invoked from the editor and hot reloads invoked from within the IDE.
-				string TargetName = TargetDesc.TargetName;
-				if( !TargetName.EndsWith( "Editor", StringComparison.InvariantCultureIgnoreCase ) && TargetDesc.bIsEditorRecompile )
-				{
-					TargetName += "Editor";
-				}
-
-                UBTMakefilePath = Path.Combine(PlatformIntermediatePath, TargetName, TargetDescs[0].Configuration.ToString(), UBTMakefileName );
+				UBTMakefilePath = Path.Combine(GetUBTMakefileDirectoryPathForSingleTarget(TargetDesc), UBTMakefileName);
             }
             else
             {
@@ -2382,6 +2405,33 @@ namespace UnrealBuildTool
             return UBTMakefilePath;
         }
 
+		/// <summary>
+		/// Gets the file path for a UBTMakefile for single target.
+		/// </summary>
+		/// <param name="Target">The target.</param>
+		/// <returns>UBTMakefile path</returns>
+		private static string GetUBTMakefileDirectoryPathForSingleTarget(TargetDescriptor Target)
+		{
+			// If there's only one target, just save the UBTMakefile in the target's build intermediate directory
+			// under a folder for that target (and platform/config combo.)
+			string PlatformIntermediatePath = BuildConfiguration.PlatformIntermediatePath;
+			if (UnrealBuildTool.HasUProjectFile())
+			{
+				PlatformIntermediatePath = Path.Combine(UnrealBuildTool.GetUProjectPath(), BuildConfiguration.PlatformIntermediateFolder);
+			}
+
+			// @todo ubtmake: If this is a compile triggered from the editor it will have passed along the game's target name, not the editor target name.
+			// At this point in Unreal Build Tool, we can't possibly know what the actual editor target name is, but we can take a guess.
+			// Even if we get it wrong, this won't have any side effects aside from not being able to share the exact same cached UBT Makefile
+			// between hot reloads invoked from the editor and hot reloads invoked from within the IDE.
+			string TargetName = Target.TargetName;
+			if (!TargetName.EndsWith("Editor", StringComparison.InvariantCultureIgnoreCase) && Target.bIsEditorRecompile)
+			{
+				TargetName += "Editor";
+			}
+
+			return Path.Combine(PlatformIntermediatePath, TargetName, Target.Configuration.ToString());
+		}
 
         /// <summary>
         /// Makes up a name for a set of targets that we can use for file or directory names
diff --git a/Engine/Source/Programs/UnrealHeaderTool/Private/CodeGenerator.cpp b/Engine/Source/Programs/UnrealHeaderTool/Private/CodeGenerator.cpp
--- a/Engine/Source/Programs/UnrealHeaderTool/Private/CodeGenerator.cpp
+++ b/Engine/Source/Programs/UnrealHeaderTool/Private/CodeGenerator.cpp
@@ -1544,6 +1544,37 @@ void WriteMacro(FStringOutputDevice &Output, const FString& MacroName, const FSt
 	Output.Log(*Macroize(*MacroName, *MacroContent));
 }
 
+/**
+ * Writes to output device auto-includes for the given source file.
+ *
+ * @param Out Output device.
+ * @param SourceFile Source file.
+ */
+void ExportAutoIncludes(FStringOutputDevice& Out, const FUnrealSourceFile& SourceFile)
+{
+	for (const auto& Include : SourceFile.GetIncludes())
+	{
+		if (!Include.IsAutoInclude())
+		{
+			continue;
+		}
+
+		const auto* AutoIncludedSourceFile = Include.GetResolved();
+
+		if (AutoIncludedSourceFile == nullptr)
+		{
+			continue;
+		}
+
+		Out.Logf(
+			TEXT("#ifndef %s")			LINE_TERMINATOR
+			TEXT("	#include \"%s\"")	LINE_TERMINATOR
+			TEXT("#endif")				LINE_TERMINATOR
+			LINE_TERMINATOR,
+			*AutoIncludedSourceFile->GetFileDefineName(), *AutoIncludedSourceFile->GetIncludePath());
+	}
+}
+
 void FNativeClassHeaderGenerator::ExportClassesFromSourceFileInner(FUnrealSourceFile& SourceFile)
 {
 	TArray<UEnum*>				Enums;
@@ -1551,12 +1582,14 @@ void FNativeClassHeaderGenerator::ExportClassesFromSourceFileInner(FUnrealSource
 	TArray<UDelegateFunction*>	DelegateFunctions;
 
 	GeneratedHeaderText.Logf(
-		TEXT("#ifdef %s_%s_generated_h")													LINE_TERMINATOR
+		TEXT("#ifdef %s")																	LINE_TERMINATOR
 		TEXT("#error \"%s.generated.h already included, missing '#pragma once' in %s.h\"")	LINE_TERMINATOR
 		TEXT("#endif")																		LINE_TERMINATOR
-		TEXT("#define %s_%s_generated_h")													LINE_TERMINATOR
+		TEXT("#define %s")																	LINE_TERMINATOR
 		LINE_TERMINATOR,
-		*API, *SourceFile.GetStrippedFilename(), *SourceFile.GetStrippedFilename(), *SourceFile.GetStrippedFilename(), *API, *SourceFile.GetStrippedFilename());
+		*SourceFile.GetFileDefineName(), *SourceFile.GetStrippedFilename(), *SourceFile.GetStrippedFilename(), *SourceFile.GetFileDefineName());
+
+	ExportAutoIncludes(GeneratedHeaderText, SourceFile);
 
 	// get the lists of fields that belong to this class that should be exported
 	SourceFile.GetScope()->SplitTypesIntoArrays(Enums, Structs, DelegateFunctions);
@@ -2205,10 +2238,10 @@ void FNativeClassHeaderGenerator::ExportClassesFromSourceFileWrapper(FUnrealSour
 	ConvertToBuildIncludePath(Package, NewFileName);
 
 	auto IncludeStr = FString::Printf(
-		TEXT("#ifndef %s_%s_generated_h")													LINE_TERMINATOR
-		TEXT("    #include \"%s\"")															LINE_TERMINATOR
-		TEXT("#endif")																		LINE_TERMINATOR,
-		*API, *SourceFile.GetStrippedFilename(), *NewFileName);
+		TEXT("#ifndef %s")			LINE_TERMINATOR
+		TEXT("    #include \"%s\"")	LINE_TERMINATOR
+		TEXT("#endif")				LINE_TERMINATOR,
+		*SourceFile.GetFileDefineName(), *NewFileName);
 
 	// Keep track of all of the UObject headers for this module, in the same order that we digest them in
 	// @todo uht: We're wrapping these includes in checks for header guards, ONLY because most existing UObject headers are missing '#pragma once'
diff --git a/Engine/Source/Programs/UnrealHeaderTool/Private/HeaderProvider.cpp b/Engine/Source/Programs/UnrealHeaderTool/Private/HeaderProvider.cpp
--- a/Engine/Source/Programs/UnrealHeaderTool/Private/HeaderProvider.cpp
+++ b/Engine/Source/Programs/UnrealHeaderTool/Private/HeaderProvider.cpp
@@ -5,8 +5,8 @@
 #include "UnrealSourceFile.h"
 #include "ClassMaps.h"
 
-FHeaderProvider::FHeaderProvider(EHeaderProviderSourceType Type, const FString& Id)
-	: Type(Type), Id(Id), Cache(nullptr)
+FHeaderProvider::FHeaderProvider(EHeaderProviderSourceType Type, const FString& Id, bool bAutoInclude/* = false*/)
+	: Type(Type), Id(Id), Cache(nullptr), bAutoInclude(bAutoInclude)
 {
 
 }
@@ -86,6 +86,12 @@ EHeaderProviderSourceType FHeaderProvider::GetType() const
 	return Type;
 }
 
+const FUnrealSourceFile* FHeaderProvider::GetResolved() const
+{
+	check(Type == EHeaderProviderSourceType::Resolved);
+	return Cache;
+}
+
 bool operator==(const FHeaderProvider& A, const FHeaderProvider& B)
 {
 	if (A.Cache != B.Cache)
diff --git a/Engine/Source/Programs/UnrealHeaderTool/Private/HeaderProvider.h b/Engine/Source/Programs/UnrealHeaderTool/Private/HeaderProvider.h
--- a/Engine/Source/Programs/UnrealHeaderTool/Private/HeaderProvider.h
+++ b/Engine/Source/Programs/UnrealHeaderTool/Private/HeaderProvider.h
@@ -17,9 +17,10 @@ class FHeaderProvider
 {
 	friend bool operator==(const FHeaderProvider& A, const FHeaderProvider& B);
 public:
-	FHeaderProvider(EHeaderProviderSourceType Type, const FString& Id);
+	FHeaderProvider(EHeaderProviderSourceType Type, const FString& Id, bool bAutoInclude = false);
 
 	FUnrealSourceFile* Resolve();
+	const FUnrealSourceFile* GetResolved() const;
 
 	FString ToString() const;
 
@@ -27,10 +28,15 @@ public:
 
 	EHeaderProviderSourceType GetType() const;
 
+	bool IsAutoInclude() const { return bAutoInclude; }
+
 private:
 	EHeaderProviderSourceType Type;
 	FString Id;
 	FUnrealSourceFile* Cache;
+
+	// Tells if this include should be auto included in generated.h file.
+	bool bAutoInclude;
 };
 
 bool operator==(const FHeaderProvider& A, const FHeaderProvider& B);
\ No newline at end of file
diff --git a/Engine/Source/Programs/UnrealHeaderTool/Private/UnrealSourceFile.cpp b/Engine/Source/Programs/UnrealHeaderTool/Private/UnrealSourceFile.cpp
--- a/Engine/Source/Programs/UnrealHeaderTool/Private/UnrealSourceFile.cpp
+++ b/Engine/Source/Programs/UnrealHeaderTool/Private/UnrealSourceFile.cpp
@@ -108,3 +108,13 @@ bool FUnrealSourceFile::HasChanged() const
 {
 	return bHasChanged;
 }
+
+FString FUnrealSourceFile::GetAPI() const
+{
+	return FPackageName::GetShortName(Package).ToUpper();
+}
+
+FString FUnrealSourceFile::GetFileDefineName() const
+{
+	return FString::Printf(TEXT("%s_%s_generated_h"), *GetAPI(), *GetStrippedFilename());
+}
diff --git a/Engine/Source/Programs/UnrealHeaderTool/Private/UnrealSourceFile.h b/Engine/Source/Programs/UnrealHeaderTool/Private/UnrealSourceFile.h
--- a/Engine/Source/Programs/UnrealHeaderTool/Private/UnrealSourceFile.h
+++ b/Engine/Source/Programs/UnrealHeaderTool/Private/UnrealSourceFile.h
@@ -82,6 +82,16 @@ public:
 	 */
 	FString GetFileId() const;
 
+	/**
+	 * Gets source file API.
+	 */
+	FString GetAPI() const;
+
+	/**
+	 * Gets define name of this source file.
+	 */
+	FString GetFileDefineName() const;
+
 	/**
 	 * Gets file-wise generated body macro name.
 	 *
@@ -166,6 +176,14 @@ public:
 		return Includes;
 	}
 
+	/**
+	 * Gets includes. Const version.
+	 */
+	const TArray<FHeaderProvider>& GetIncludes() const
+	{
+		return Includes;
+	}
+
 	/**
 	 * Sets generated filename.
 	 */
diff --git a/Engine/Source/Runtime/Projects/Private/ModuleDescriptor.cpp b/Engine/Source/Runtime/Projects/Private/ModuleDescriptor.cpp
--- a/Engine/Source/Runtime/Projects/Private/ModuleDescriptor.cpp
+++ b/Engine/Source/Runtime/Projects/Private/ModuleDescriptor.cpp
@@ -154,6 +154,17 @@ bool FModuleDescriptor::Read(const FJsonObject& Object, FText& OutFailReason)
 		}
 	}
 
+	// Read the additional dependencies
+	TSharedPtr<FJsonValue> AdditionalDependenciesValue = Object.TryGetField(TEXT("AdditionalDependencies"));
+	if (AdditionalDependenciesValue.IsValid() && AdditionalDependenciesValue->Type == EJson::Array)
+	{
+		const TArray< TSharedPtr< FJsonValue > >& DepArray = AdditionalDependenciesValue->AsArray();
+		for (int Idx = 0; Idx < DepArray.Num(); Idx++)
+		{
+			AdditionalDependencies.Add(DepArray[Idx]->AsString());
+		}
+	}
+
 	return true;
 }
 
@@ -215,6 +226,15 @@ void FModuleDescriptor::Write(TJsonWriter<>& Writer) const
 		}
 		Writer.WriteArrayEnd();
 	}
+	if (AdditionalDependencies.Num() > 0)
+	{
+		Writer.WriteArrayStart(TEXT("AdditionalDependencies"));
+		for (int Idx = 0; Idx < AdditionalDependencies.Num(); Idx++)
+		{
+			Writer.WriteValue(AdditionalDependencies[Idx]);
+		}
+		Writer.WriteArrayEnd();
+	}
 	Writer.WriteObjectEnd();
 }
 
diff --git a/Engine/Source/Runtime/Projects/Public/ModuleDescriptor.h b/Engine/Source/Runtime/Projects/Public/ModuleDescriptor.h
--- a/Engine/Source/Runtime/Projects/Public/ModuleDescriptor.h
+++ b/Engine/Source/Runtime/Projects/Public/ModuleDescriptor.h
@@ -108,6 +108,9 @@ struct PROJECTS_API FModuleDescriptor
 	/** List of disallowed platforms */
 	TArray<FString> BlacklistPlatforms;
 
+	/** List of additional dependencies for building this module. */
+	TArray<FString> AdditionalDependencies;
+
 	/** Normal constructor */
 	FModuleDescriptor(const FName InName = NAME_None, EHostType::Type InType = EHostType::Runtime, ELoadingPhase::Type InLoadingPhase = ELoadingPhase::Default);
 
diff --git a/Engine/Source/Runtime/Slate/Public/SlateBasics.h b/Engine/Source/Runtime/Slate/Public/SlateBasics.h
--- a/Engine/Source/Runtime/Slate/Public/SlateBasics.h
+++ b/Engine/Source/Runtime/Slate/Public/SlateBasics.h
@@ -7,10 +7,7 @@
 #include "Json.h"
 #include "SlateCore.h"
 #include "SlateClasses.h"
-
-
-#define BEGIN_SLATE_FUNCTION_BUILD_OPTIMIZATION PRAGMA_DISABLE_OPTIMIZATION
-#define END_SLATE_FUNCTION_BUILD_OPTIMIZATION   PRAGMA_ENABLE_OPTIMIZATION
+#include "SlateOptMacros.h"
 
 #include "IPlatformTextField.h"
 #if PLATFORM_IOS
diff --git a/Engine/Source/Runtime/Slate/Public/SlateOptMacros.h b/Engine/Source/Runtime/Slate/Public/SlateOptMacros.h
new file mode 100644
--- /dev/null
+++ b/Engine/Source/Runtime/Slate/Public/SlateOptMacros.h
@@ -0,0 +1,8 @@
+// Copyright 1998-2015 Epic Games, Inc. All Rights Reserved.
+
+#pragma once
+
+#include "CoreMiscDefines.h"
+
+#define BEGIN_SLATE_FUNCTION_BUILD_OPTIMIZATION PRAGMA_DISABLE_OPTIMIZATION
+#define END_SLATE_FUNCTION_BUILD_OPTIMIZATION   PRAGMA_ENABLE_OPTIMIZATION
\ No newline at end of file
diff --git a/Engine/Source/Runtime/SlateCore/Public/Input/NavigationReply.h b/Engine/Source/Runtime/SlateCore/Public/Input/NavigationReply.h
--- a/Engine/Source/Runtime/SlateCore/Public/Input/NavigationReply.h
+++ b/Engine/Source/Runtime/SlateCore/Public/Input/NavigationReply.h
@@ -2,6 +2,8 @@
 
 #pragma once
 
+#include "SlateEnums.h"
+
 #include "NavigationReply.generated.h"
 
 class SWidget;
diff --git a/Engine/Source/Runtime/SlateCore/Public/Layout/Children.h b/Engine/Source/Runtime/SlateCore/Public/Layout/Children.h
--- a/Engine/Source/Runtime/SlateCore/Public/Layout/Children.h
+++ b/Engine/Source/Runtime/SlateCore/Public/Layout/Children.h
@@ -2,8 +2,10 @@
 
 #pragma once
 
-class SWidget;
+#include "SlotBase.h"
+#include "Margin.h"
 
+class SWidget;
 
 /**
  * FChildren is an interface that must be implemented by all child containers.
diff --git a/Engine/Source/Runtime/SlateCore/Public/SlateCore.h b/Engine/Source/Runtime/SlateCore/Public/SlateCore.h
--- a/Engine/Source/Runtime/SlateCore/Public/SlateCore.h
+++ b/Engine/Source/Runtime/SlateCore/Public/SlateCore.h
@@ -28,8 +28,7 @@ enum class EActiveTimerReturnType : uint8;
 /* Delegates
  *****************************************************************************/
 
-/** Delegate used for registering a widget for tick services. */
-DECLARE_DELEGATE_RetVal_TwoParams(EActiveTimerReturnType, FWidgetActiveTimerDelegate, double /*InCurrentTime*/, float /*InDeltaTime*/);
+#include "WidgetActiveTimerDelegate.h"
 
 /* Includes
  *****************************************************************************/
diff --git a/Engine/Source/Runtime/SlateCore/Public/Styling/SlateWidgetStyle.h b/Engine/Source/Runtime/SlateCore/Public/Styling/SlateWidgetStyle.h
--- a/Engine/Source/Runtime/SlateCore/Public/Styling/SlateWidgetStyle.h
+++ b/Engine/Source/Runtime/SlateCore/Public/Styling/SlateWidgetStyle.h
@@ -4,6 +4,7 @@
 
 #include "SlateWidgetStyle.generated.h"
 
+struct FSlateBrush;
 
 /**
  * Base structure for widget styles.
diff --git a/Engine/Source/Runtime/SlateCore/Public/Types/WidgetActiveTimerDelegate.h b/Engine/Source/Runtime/SlateCore/Public/Types/WidgetActiveTimerDelegate.h
new file mode 100644
--- /dev/null
+++ b/Engine/Source/Runtime/SlateCore/Public/Types/WidgetActiveTimerDelegate.h
@@ -0,0 +1,6 @@
+// Copyright 1998-2015 Epic Games, Inc. All Rights Reserved.
+
+#pragma once
+
+/** Delegate used for registering a widget for tick services. */
+DECLARE_DELEGATE_RetVal_TwoParams(EActiveTimerReturnType, FWidgetActiveTimerDelegate, double /*InCurrentTime*/, float /*InDeltaTime*/);
\ No newline at end of file
diff --git a/Engine/Source/Runtime/SlateCore/Public/Widgets/DeclarativeSyntaxSupport.h b/Engine/Source/Runtime/SlateCore/Public/Widgets/DeclarativeSyntaxSupport.h
--- a/Engine/Source/Runtime/SlateCore/Public/Widgets/DeclarativeSyntaxSupport.h
+++ b/Engine/Source/Runtime/SlateCore/Public/Widgets/DeclarativeSyntaxSupport.h
@@ -2,6 +2,10 @@
 
 #pragma once
 
+#include "SNullWidget.h"
+
+class ISlateMetaData;
+class IToolTip;
 
 /**
  * Slate widgets are constructed through SNew and SAssignNew.
diff --git a/Engine/Source/Runtime/SlateCore/Public/Widgets/SCompoundWidget.h b/Engine/Source/Runtime/SlateCore/Public/Widgets/SCompoundWidget.h
--- a/Engine/Source/Runtime/SlateCore/Public/Widgets/SCompoundWidget.h
+++ b/Engine/Source/Runtime/SlateCore/Public/Widgets/SCompoundWidget.h
@@ -2,6 +2,10 @@
 
 #pragma once
 
+#include "SWidget.h"
+#include "SlateColor.h"
+#include "Children.h"
+
 
 /**
  * A CompoundWidget is the base from which most non-primitive widgets should be built.
diff --git a/Engine/Source/Runtime/SlateCore/Public/Widgets/SWidget.h b/Engine/Source/Runtime/SlateCore/Public/Widgets/SWidget.h
--- a/Engine/Source/Runtime/SlateCore/Public/Widgets/SWidget.h
+++ b/Engine/Source/Runtime/SlateCore/Public/Widgets/SWidget.h
@@ -2,7 +2,38 @@
 
 #pragma once
 
+#include "Visibility.h"
 #include "SlateRenderTransform.h"
+#include "NavigationReply.h"
+#include "SlateColor.h"
+#include "WidgetActiveTimerDelegate.h"
+#include "DeclarativeSyntaxSupport.h"
+
+class ISlateMetaData;
+class FActiveTimerHandle;
+class FPaintArgs;
+class FSlateRect;
+class FSlateWindowElementList;
+class FWidgetStyle;
+class FWeakWidgetPath;
+class FWidgetPath;
+class FDragDropEvent;
+class FSlotBase;
+class FArrangedChildren;
+class FChildren;
+class FArrangedWidget;
+struct FSlateBrush;
+struct FGeometry;
+struct FFocusEvent;
+struct FKeyboardFocusEvent;
+struct FCharacterEvent;
+struct FKeyEvent;
+struct FControllerEvent;
+struct FAnalogInputEvent;
+struct FPointerEvent;
+struct FMotionEvent;
+struct FVirtualPointerPosition;
+struct FNavigationEvent;
 
 enum class EPopupMethod : uint8;
 
@@ -435,7 +466,7 @@ public:
 	 */
 	virtual TOptional<EPopupMethod> OnQueryPopupMethod() const;
 
-	virtual TSharedPtr<struct FVirtualPointerPosition> TranslateMouseCoordinateFor3DChild(const TSharedRef<SWidget>& ChildWidget, const FGeometry& MyGeometry, const FVector2D& ScreenSpaceMouseCoordinate, const FVector2D& LastScreenSpaceMouseCoordinate) const;
+	virtual TSharedPtr<FVirtualPointerPosition> TranslateMouseCoordinateFor3DChild(const TSharedRef<SWidget>& ChildWidget, const FGeometry& MyGeometry, const FVector2D& ScreenSpaceMouseCoordinate, const FVector2D& LastScreenSpaceMouseCoordinate) const;
 	
 	/**
 	 * All the pointer (mouse, touch, stylus, etc.) events from this frame have been routed.
