Remove Windows resource files from generated game projects. Users can create resource files themselves if they need platform specific resources; we already provide a platform agnostic way of customizing the icon through the project settings, and any version strings would be better implemented in a similar platform-agnostic way. The current resources are full of Epic copyright and version information, can't be upgraded, and can't be opened with the Visual Studio resource editor.

diff --git a/Engine/Source/Editor/GameProjectGeneration/Private/GameProjectUtils.cpp b/Engine/Source/Editor/GameProjectGeneration/Private/GameProjectUtils.cpp
--- a/Engine/Source/Editor/GameProjectGeneration/Private/GameProjectUtils.cpp
+++ b/Engine/Source/Editor/GameProjectGeneration/Private/GameProjectUtils.cpp
@@ -1507,14 +1507,6 @@ bool GameProjectUtils::CreateProjectFromTemplate(const FProjectInformation& InPr
 
 	SlowTask.EnterProgressFrame();
 
-	// Copy resources
-	const FString GameModuleSourcePath = DestFolder / TEXT("Source") / ProjectName;
-	if (GenerateGameResourceFiles(GameModuleSourcePath, ProjectName, DestFolder, InProjectInfo.bShouldGenerateCode, CreatedFiles, OutFailReason) == false)
-	{
-		DeleteCreatedFiles(DestFolder, CreatedFiles);
-		return false;
-	}
-
 	SlowTask.EnterProgressFrame();
 	if ( InProjectInfo.bShouldGenerateCode )
 	{
@@ -1816,12 +1808,6 @@ bool GameProjectUtils::GenerateBasicSourceCode(const FString& NewProjectSourcePa
 		}
 	}
 
-	// MyGame resource folder
-	if (GenerateGameResourceFiles(GameModulePath, NewProjectName, NewProjectRoot, true, OutCreatedFiles, OutFailReason) == false)
-	{
-		return false;
-	}
-
 	// MyGame.Target.cs
 	{
 		const FString NewTargetFilename = NewProjectSourcePath / NewProjectName + TEXT(".Target.cs");
@@ -2659,56 +2645,6 @@ bool GameProjectUtils::GenerateGameModuleTargetFile(const FString& NewBuildFileN
 	return WriteOutputFile(NewBuildFileName, FinalOutput, OutFailReason);
 }
 
-bool GameProjectUtils::GenerateGameResourceFiles(const FString& NewResourceFolderName, const FString& GameName, const FString& GameRoot, bool bShouldGenerateCode, TArray<FString>& OutCreatedFiles, FText& OutFailReason)
-{
-#if PLATFORM_WINDOWS
-	// Copy the icon if it doesn't already exist. If we're upgrading a content-only project to code, it will already have one unless it was created before content-only project icons were supported.
-	FString IconFileName = GameRoot / TEXT("Build/Windows/Application.ico");
-	if(!FPaths::FileExists(IconFileName))
-	{
-		SourceControlHelpers::CopyFileUnderSourceControl(IconFileName, FPaths::EngineContentDir() / TEXT("Editor/Templates/Resources/Windows/_GAME_NAME_.ico"), LOCTEXT("IconFileDescription", "icon"), OutFailReason);
-		OutCreatedFiles.Add(IconFileName);
-	}
-	
-	// Generate a RC script if it's a code project
-	if(bShouldGenerateCode)
-	{
-		FString OutputFilename = NewResourceFolderName / FString::Printf(TEXT("Resources/Windows/%s.rc"), *GameName);
-
-		FString TemplateText;
-		if (!ReadTemplateFile(TEXT("Resources/Windows/_GAME_NAME_.rc"), TemplateText, OutFailReason))
-		{
-			return false;
-		}
-
-		FString RelativeIconPath = IconFileName;
-		FPaths::MakePathRelativeTo(RelativeIconPath, *OutputFilename);
-		TemplateText = TemplateText.Replace(TEXT("%ICON_PATH%"), *RelativeIconPath, ESearchCase::CaseSensitive);
-		TemplateText = TemplateText.Replace(TEXT("%GAME_NAME%"), *GameName, ESearchCase::CaseSensitive);
-
-		struct Local
-		{
-			static bool WriteFile(const FString& InDestFile, const FText& InFileDescription, FText& OutFailureReason, FString* InFileContents, TArray<FString>* OutCreatedFileList)
-			{
-				if (WriteOutputFile(InDestFile, *InFileContents, OutFailureReason))
-				{
-					OutCreatedFileList->Add(InDestFile);
-					return true;
-				}
-
-				return false;
-			}
-		};
-
-		SourceControlHelpers::CheckoutOrMarkForAdd(OutputFilename, LOCTEXT("ResourceFileDescription", "resource"), FOnPostCheckOut::CreateStatic(&Local::WriteFile, &TemplateText, &OutCreatedFiles), OutFailReason);
-	}
-#elif PLATFORM_MAC
-	//@todo MAC: Implement MAC version of these files...
-#endif
-
-	return true;
-}
-
 bool GameProjectUtils::GenerateEditorModuleBuildFile(const FString& NewBuildFileName, const FString& ModuleName, const TArray<FString>& PublicDependencyModuleNames, const TArray<FString>& PrivateDependencyModuleNames, FText& OutFailReason)
 {
 	FString Template;
diff --git a/Engine/Source/Editor/GameProjectGeneration/Private/GameProjectUtils.h b/Engine/Source/Editor/GameProjectGeneration/Private/GameProjectUtils.h
--- a/Engine/Source/Editor/GameProjectGeneration/Private/GameProjectUtils.h
+++ b/Engine/Source/Editor/GameProjectGeneration/Private/GameProjectUtils.h
@@ -132,9 +132,6 @@ public:
 	/** Creates code project files for a new game project. On failure, OutFailReason will be populated. */
 	static bool GenerateCodeProjectFiles(const FString& ProjectFilename, FText& OutFailReason);
 
-	/** Generates a set of resource files for a game module */
-	static bool GenerateGameResourceFiles(const FString& NewResourceFolderName, const FString& GameName, const FString& GameRoot, bool bShouldGenerateCode, TArray<FString>& OutCreatedFiles, FText& OutFailReason);
-
 	/** Returns true if there are starter content files available for instancing into new projects. */
 	static bool IsStarterContentAvailableForNewProjects();
 
