Replaced use of IsRocket for checking whether to open an IDE with IsEngineInstalled, which gives clearer reasoning for this.

diff --git a/Engine/Source/Editor/GameProjectGeneration/Private/SNewProjectWizard.cpp b/Engine/Source/Editor/GameProjectGeneration/Private/SNewProjectWizard.cpp
--- a/Engine/Source/Editor/GameProjectGeneration/Private/SNewProjectWizard.cpp
+++ b/Engine/Source/Editor/GameProjectGeneration/Private/SNewProjectWizard.cpp
@@ -1484,9 +1484,9 @@ void SNewProjectWizard::CreateAndOpenProject( )
 
 	if( GetSelectedTemplateItem()->bGenerateCode )
 	{
-	    // Rocket already has the engine compiled, so we can try to build and open a new project immediately. Non-Rocket might require building
+	    // If the engine is installed it is already compiled, so we can try to build and open a new project immediately. Non-installed situations might require building
 	    // the engine (especially the case when binaries came from P4), so we only open the IDE for that.
-		if (FRocketSupport::IsRocket())
+		if (FApp::IsEngineInstalled())
 		{
 			if (GameProjectUtils::BuildCodeProject(ProjectFile))
 			{
