Hide the 'Convert in-place' and 'Skip conversion' buttons in the project conversion dialog by default. Hit 'More Options...' to summon them.

diff --git a/Engine/Source/Editor/GameProjectGeneration/Private/SProjectBrowser.cpp b/Engine/Source/Editor/GameProjectGeneration/Private/SProjectBrowser.cpp
--- a/Engine/Source/Editor/GameProjectGeneration/Private/SProjectBrowser.cpp
+++ b/Engine/Source/Editor/GameProjectGeneration/Private/SProjectBrowser.cpp
@@ -977,22 +977,36 @@ bool SProjectBrowser::OpenProject( const FString& InProjectFile )
 			return false;
 		}
 
+		// Hyperlinks for the upgrade dialog
+		TArray<FText> Hyperlinks;
+		int32 MoreOptionsHyperlink = Hyperlinks.Add(LOCTEXT("ProjectConvert_MoreOptions", "More Options..."));
+
 		// Button labels for the upgrade dialog
 		TArray<FText> Buttons;
 		int32 OpenCopyButton = Buttons.Add(LOCTEXT("ProjectConvert_OpenCopy", "Open a copy"));
-		int32 OpenExistingButton = Buttons.Add(LOCTEXT("ProjectConvert_ConvertInPlace", "Convert in-place"));
-		int32 SkipConversionButton = Buttons.Add(LOCTEXT("ProjectConvert_SkipConversion", "Skip conversion"));
 		int32 CancelButton = Buttons.Add(LOCTEXT("ProjectConvert_Cancel", "Cancel"));
+		int32 OpenExistingButton = -1;
+		int32 SkipConversionButton = -1;
 
 		// Prompt for upgrading. Different message for code and content projects, since the process is a bit trickier for code.
-		int32 Selection;
+		FText DialogText;
 		if(ProjectStatus.bCodeBasedProject)
 		{
-			Selection = SVerbChoiceDialog::ShowModal(LOCTEXT("ProjectConversionTitle", "Convert Project"), LOCTEXT("ConvertCodeProjectPrompt", "This project was made with a different version of the Unreal Engine. Converting to this version will rebuild your code projects.\n\nNew features and improvements sometimes cause API changes, which may require you to modify your code before it compiles. Content saved with newer versions of the editor will not open in older versions.\n\nWe recommend you open a copy of your project to avoid damaging the original."), Buttons);
+			DialogText = LOCTEXT("ConvertCodeProjectPrompt", "This project was made with a different version of the Unreal Engine. Converting to this version will rebuild your code projects.\n\nNew features and improvements sometimes cause API changes, which may require you to modify your code before it compiles. Content saved with newer versions of the editor will not open in older versions.\n\nWe recommend you open a copy of your project to avoid damaging the original.");
 		}
 		else
 		{
-			Selection = SVerbChoiceDialog::ShowModal(LOCTEXT("ProjectConversionTitle", "Convert Project"), LOCTEXT("ConvertContentProjectPrompt", "This project was made with a different version of the Unreal Engine.\n\nOpening it with this version of the editor may prevent it opening with the original editor, and may lose data. We recommend you open a copy to avoid damaging the original."), Buttons);
+			DialogText = LOCTEXT("ConvertContentProjectPrompt", "This project was made with a different version of the Unreal Engine.\n\nOpening it with this version of the editor may prevent it opening with the original editor, and may lose data. We recommend you open a copy to avoid damaging the original.");
+		}
+
+		// Show the dialog, and expand to the advanced dialog if the user selects 'More Options...'
+		int32 Selection = SVerbChoiceDialog::ShowModal(LOCTEXT("ProjectConversionTitle", "Convert Project"), DialogText, Hyperlinks, Buttons);
+		if(~Selection == MoreOptionsHyperlink)
+		{
+			OpenExistingButton = Buttons.Insert(LOCTEXT("ProjectConvert_ConvertInPlace", "Convert in-place"), 1);
+			SkipConversionButton = Buttons.Insert(LOCTEXT("ProjectConvert_SkipConversion", "Skip conversion"), 2);
+			CancelButton += 2;
+			Selection = SVerbChoiceDialog::ShowModal(LOCTEXT("ProjectConversionTitle", "Convert Project"), DialogText, Buttons);
 		}
 
 		// Handle the selection
diff --git a/Engine/Source/Editor/GameProjectGeneration/Private/SVerbChoiceDialog.cpp b/Engine/Source/Editor/GameProjectGeneration/Private/SVerbChoiceDialog.cpp
--- a/Engine/Source/Editor/GameProjectGeneration/Private/SVerbChoiceDialog.cpp
+++ b/Engine/Source/Editor/GameProjectGeneration/Private/SVerbChoiceDialog.cpp
@@ -4,7 +4,13 @@
 #include "SVerbChoiceDialog.h"
 #include "SHyperlink.h"
 
-int32 SVerbChoiceDialog::ShowModal( const FText& InTitle, const FText& InMessage, const TArray<FText>& InButtons )
+int32 SVerbChoiceDialog::ShowModal(const FText& InTitle, const FText& InText, const TArray<FText>& InButtons)
+{
+	TArray<FText> Hyperlinks;
+	return ShowModal(InTitle, InText, Hyperlinks, InButtons);
+}
+
+int32 SVerbChoiceDialog::ShowModal( const FText& InTitle, const FText& InMessage, const TArray<FText>& InHyperlinks, const TArray<FText>& InButtons)
 {
 	TSharedRef<SWindow> ModalWindow = SNew(SWindow)
 		.Title( InTitle )
@@ -16,6 +22,7 @@ int32 SVerbChoiceDialog::ShowModal( const FText& InTitle, const FText& InMessage
 	TSharedRef<SVerbChoiceDialog> MessageBox = SNew(SVerbChoiceDialog)
 		.ParentWindow(ModalWindow)
 		.Message( InMessage )
+		.Hyperlinks( InHyperlinks )
 		.Buttons( InButtons )
 		.WrapMessageAt(640.0f);
 
@@ -35,9 +42,11 @@ void SVerbChoiceDialog::Construct( const FArguments& InArgs )
 
 	FSlateFontInfo MessageFont( FEditorStyle::GetFontStyle("StandardDialog.LargeFont"));
 	Message = InArgs._Message;
+	Hyperlinks = InArgs._Hyperlinks;
 	Buttons = InArgs._Buttons;
-
+	
 	TSharedPtr<SUniformGridPanel> ButtonBox;
+	TSharedPtr<SUniformGridPanel> HyperlinksBox;
 
 	this->ChildSlot
 		[	
@@ -74,19 +83,19 @@ void SVerbChoiceDialog::Construct( const FArguments& InArgs )
 								.FillWidth(1.0f)
 								.HAlign(HAlign_Left)
 								.VAlign(VAlign_Center)
-								.Padding(12.0f)
+								.Padding(5.0f)
 								[
-									SNew(SHyperlink)
-										.OnNavigate(this, &SVerbChoiceDialog::HandleCopyMessageHyperlinkNavigate)
-										.Text( NSLOCTEXT("SVerbChoiceDialog", "CopyMessageHyperlink", "Copy Message") )
-										.ToolTipText( NSLOCTEXT("SVerbChoiceDialog", "CopyMessageTooltip", "Copy the text in this message to the clipboard (CTRL+C)") )
+									SAssignNew( HyperlinksBox, SUniformGridPanel )
+										.SlotPadding(FEditorStyle::GetMargin("StandardDialog.SlotPadding"))
+										.MinDesiredSlotWidth(FEditorStyle::GetFloat("StandardDialog.MinDesiredSlotWidth"))
+										.MinDesiredSlotHeight(FEditorStyle::GetFloat("StandardDialog.MinDesiredSlotHeight"))
 								]
 
 							+ SHorizontalBox::Slot()
 								.AutoWidth()
 								.HAlign(HAlign_Right)
 								.VAlign(VAlign_Center)
-								.Padding(2.f)
+								.Padding(5.0f)
 								[
 									SAssignNew( ButtonBox, SUniformGridPanel )
 										.SlotPadding(FEditorStyle::GetMargin("StandardDialog.SlotPadding"))
@@ -97,6 +106,24 @@ void SVerbChoiceDialog::Construct( const FArguments& InArgs )
 				]
 		];
 
+	HyperlinksBox->AddSlot(0, 0)
+		[
+			SNew(SHyperlink)
+				.OnNavigate(this, &SVerbChoiceDialog::HandleCopyMessageHyperlinkNavigate)
+				.Text( NSLOCTEXT("SVerbChoiceDialog", "CopyMessageHyperlink", "Copy Message") )
+				.ToolTipText( NSLOCTEXT("SVerbChoiceDialog", "CopyMessageTooltip", "Copy the text in this message to the clipboard (CTRL+C)") )
+		];
+
+	for(int32 Idx = 0; Idx < Hyperlinks.Get().Num(); Idx++)
+	{
+		HyperlinksBox->AddSlot(Idx + 1, 0)
+			[
+				SNew(SHyperlink)
+					.Text( Hyperlinks.Get()[Idx] )
+					.OnNavigate(this, &SVerbChoiceDialog::HandleHyperlinkClicked, ~Idx)
+			];
+	}
+
 	for(int32 Idx = 0; Idx < Buttons.Get().Num(); Idx++)
 	{
 		ButtonBox->AddSlot(Idx, 0)
@@ -140,15 +167,21 @@ void SVerbChoiceDialog::CopyMessageToClipboard( )
 	FPlatformMisc::ClipboardCopy( *Message.Get().ToString() );
 }
 
-FReply SVerbChoiceDialog::HandleButtonClicked( int32 InResponse )
+void SVerbChoiceDialog::HandleCopyMessageHyperlinkNavigate( )
+{
+	CopyMessageToClipboard();
+}
+
+void SVerbChoiceDialog::HandleHyperlinkClicked( int32 InResponse )
 {
 	Response = InResponse;
 	ParentWindow->RequestDestroyWindow();
-
-	return FReply::Handled();
 }
 
-void SVerbChoiceDialog::HandleCopyMessageHyperlinkNavigate( )
+FReply SVerbChoiceDialog::HandleButtonClicked( int32 InResponse )
 {
-	CopyMessageToClipboard();
+	Response = InResponse;
+	ParentWindow->RequestDestroyWindow();
+
+	return FReply::Handled();
 }
diff --git a/Engine/Source/Editor/GameProjectGeneration/Private/SVerbChoiceDialog.h b/Engine/Source/Editor/GameProjectGeneration/Private/SVerbChoiceDialog.h
--- a/Engine/Source/Editor/GameProjectGeneration/Private/SVerbChoiceDialog.h
+++ b/Engine/Source/Editor/GameProjectGeneration/Private/SVerbChoiceDialog.h
@@ -8,6 +8,7 @@ public:
 	SLATE_BEGIN_ARGS( SVerbChoiceDialog )	{}
 		SLATE_ATTRIBUTE(TSharedPtr<SWindow>, ParentWindow)
 		SLATE_ATTRIBUTE(FText, Message)	
+		SLATE_ATTRIBUTE(TArray<FText>, Hyperlinks)
 		SLATE_ATTRIBUTE(TArray<FText>, Buttons)
 		SLATE_ATTRIBUTE(float, WrapMessageAt)
 	SLATE_END_ARGS()
@@ -15,6 +16,9 @@ public:
 	/** Displays the modal dialog box */
 	static int32 ShowModal(const FText& InTitle, const FText& InText, const TArray<FText>& InButtons);
 
+	/** Displays the modal dialog box, also allowing customization of the hyperlinks */
+	static int32 ShowModal(const FText& InTitle, const FText& InText, const TArray<FText>& InHyperlinks, const TArray<FText>& InButtons);
+
 	/** Constructs the dialog */
 	void Construct( const FArguments& InArgs );
 
@@ -29,14 +33,18 @@ protected:
 	void CopyMessageToClipboard( );
 
 private:
+	/** Handles clicking the 'Copy Message' hyperlink. */
+	void HandleCopyMessageHyperlinkNavigate( );
+
+	/** Handles clicking on a hyperlink. */
+	void HandleHyperlinkClicked( int32 InResponse );
+
 	/** Handles clicking a message box button. */
 	FReply HandleButtonClicked( int32 InResponse );
 
-	/** Handles clicking the 'Copy Message' hyper link. */
-	void HandleCopyMessageHyperlinkNavigate( );
-
 	int32 Response;
 	TSharedPtr<SWindow> ParentWindow;
 	TAttribute<FText> Message;
 	TAttribute< TArray<FText> > Buttons;
+	TAttribute< TArray<FText> > Hyperlinks;
 };
