Added code to abort opening a project if the project's PackageFileUE4Version is higher than the current engine's.

diff --git a/Engine/Source/Developer/DesktopPlatform/Private/DesktopPlatformBase.cpp b/Engine/Source/Developer/DesktopPlatform/Private/DesktopPlatformBase.cpp
--- a/Engine/Source/Developer/DesktopPlatform/Private/DesktopPlatformBase.cpp
+++ b/Engine/Source/Developer/DesktopPlatform/Private/DesktopPlatformBase.cpp
@@ -198,6 +198,26 @@ bool FDesktopPlatformBase::IsStockEngineRelease(const FString &Identifier)
 	return !FGuid::Parse(Identifier, Guid);
 }
 
+bool FDesktopPlatformBase::TryParseStockEngineVersion(const FString& Identifier, FEngineVersion& OutVersion)
+{
+	TCHAR* End;
+
+	uint64 Major = FCString::Strtoui64(*Identifier, &End, 10);
+	if (Major > MAX_uint16 || *(End++) != '.')
+	{
+		return false;
+	}
+
+	uint64 Minor = FCString::Strtoui64(End, &End, 10);
+	if (Minor > MAX_uint16 || *End != 0)
+	{
+		return false;
+	}
+
+	OutVersion = FEngineVersion(Major, Minor, 0, 0, TEXT(""));
+	return true;
+}
+
 bool FDesktopPlatformBase::IsSourceDistribution(const FString &EngineRootDir)
 {
 	// Check for the existence of a SourceBuild.txt file
diff --git a/Engine/Source/Developer/DesktopPlatform/Private/DesktopPlatformBase.h b/Engine/Source/Developer/DesktopPlatform/Private/DesktopPlatformBase.h
--- a/Engine/Source/Developer/DesktopPlatform/Private/DesktopPlatformBase.h
+++ b/Engine/Source/Developer/DesktopPlatform/Private/DesktopPlatformBase.h
@@ -26,6 +26,7 @@ public:
 	virtual bool IsPreferredEngineIdentifier(const FString &Identifier, const FString &OtherIdentifier) override;
 
 	virtual bool IsStockEngineRelease(const FString &Identifier) override;
+	virtual bool TryParseStockEngineVersion(const FString& Identifier, FEngineVersion& OutVersion) override;
 	virtual bool IsSourceDistribution(const FString &RootDir) override;
 	virtual bool IsPerforceBuild(const FString &RootDir) override;
 	virtual bool IsValidRootDirectory(const FString &RootDir) override;
diff --git a/Engine/Source/Developer/DesktopPlatform/Public/IDesktopPlatform.h b/Engine/Source/Developer/DesktopPlatform/Public/IDesktopPlatform.h
--- a/Engine/Source/Developer/DesktopPlatform/Public/IDesktopPlatform.h
+++ b/Engine/Source/Developer/DesktopPlatform/Public/IDesktopPlatform.h
@@ -2,6 +2,9 @@
 
 #pragma once
 
+#include "EngineVersion.h"
+
+
 namespace EFileDialogFlags
 {
 	enum Type
@@ -240,6 +243,15 @@ public:
 	*/
 	virtual bool IsStockEngineRelease(const FString &Identifier) = 0;
 
+	/**
+	 * Attempt to get the engine version from the supplied identifier
+	 *
+	 * @param	Identifier			Engine identifier to check
+	 * @param	OutVersion			Engine version obtained from identifier
+	 * @return	true if the engine version was successfully obtained
+	 */
+	virtual bool TryParseStockEngineVersion(const FString& Identifier, FEngineVersion& OutVersion) = 0;
+
 	/**
 	* Tests whether an engine installation is a source distribution.
 	*
diff --git a/Engine/Source/Editor/GameProjectGeneration/Private/SProjectBrowser.cpp b/Engine/Source/Editor/GameProjectGeneration/Private/SProjectBrowser.cpp
--- a/Engine/Source/Editor/GameProjectGeneration/Private/SProjectBrowser.cpp
+++ b/Engine/Source/Editor/GameProjectGeneration/Private/SProjectBrowser.cpp
@@ -964,6 +964,17 @@ bool SProjectBrowser::OpenProject( const FString& InProjectFile )
 	// Get the identifier for the project
 	FString ProjectIdentifier;
 	FDesktopPlatformModule::Get()->GetEngineIdentifierForProject(ProjectFile, ProjectIdentifier);
+
+	// Abort straight away if the project engine version is newer than the current engine version
+	FEngineVersion EngineVersion;
+	if (FDesktopPlatformModule::Get()->TryParseStockEngineVersion(ProjectIdentifier, EngineVersion))
+	{
+		if (FEngineVersion::GetNewest(EngineVersion, GEngineVersion, nullptr) == EVersionComparison::First)
+		{
+			FMessageDialog::Open(EAppMsgType::Ok, LOCTEXT("CantLoadNewerProject", "Unable to open this project, as it was made with a newer version of the Unreal Engine."));
+			return false;
+		}
+	}
 	
 	// Get the identifier for the current engine
 	FString CurrentIdentifier = FDesktopPlatformModule::Get()->GetCurrentEngineIdentifier();
