Improved the error message when adding a new class fails due to a compile error

The fact that the error message said "failed to add class" despite the class itself being added was causing confusion.

In the case where adding a new class fails to hot-reload (probably due a compiler error) we now state that the class itself has been added, but that you'll need to recompile the module before it will appear in the editor, and offer an easy way to open the output log to see more detailed compiler output. We also automatically close the class wizard so that you're not left looking at an error message saying that "the class name is already in use".

diff --git a/Engine/Source/Editor/GameProjectGeneration/Private/GameProjectUtils.cpp b/Engine/Source/Editor/GameProjectGeneration/Private/GameProjectUtils.cpp
--- a/Engine/Source/Editor/GameProjectGeneration/Private/GameProjectUtils.cpp
+++ b/Engine/Source/Editor/GameProjectGeneration/Private/GameProjectUtils.cpp
@@ -956,9 +956,9 @@ bool GameProjectUtils::IsValidBaseClassForCreation_Internal(const UClass* InClas
 	return !bIsBlueprintClass && (!bNeedsAPI || bHasAPI) && !bIsInterface;
 }
 
-bool GameProjectUtils::AddCodeToProject(const FString& NewClassName, const FString& NewClassPath, const FModuleContextInfo& ModuleInfo, const FNewClassInfo ParentClassInfo, const TSet<FString>& DisallowedHeaderNames, FString& OutHeaderFilePath, FString& OutCppFilePath, FText& OutFailReason)
+GameProjectUtils::EAddCodeToProjectResult GameProjectUtils::AddCodeToProject(const FString& NewClassName, const FString& NewClassPath, const FModuleContextInfo& ModuleInfo, const FNewClassInfo ParentClassInfo, const TSet<FString>& DisallowedHeaderNames, FString& OutHeaderFilePath, FString& OutCppFilePath, FText& OutFailReason)
 {
-	const bool bAddCodeSuccessful = AddCodeToProject_Internal(NewClassName, NewClassPath, ModuleInfo, ParentClassInfo, DisallowedHeaderNames, OutHeaderFilePath, OutCppFilePath, OutFailReason);
+	const EAddCodeToProjectResult Result = AddCodeToProject_Internal(NewClassName, NewClassPath, ModuleInfo, ParentClassInfo, DisallowedHeaderNames, OutHeaderFilePath, OutCppFilePath, OutFailReason);
 
 	if( FEngineAnalytics::IsAvailable() )
 	{
@@ -966,12 +966,13 @@ bool GameProjectUtils::AddCodeToProject(const FString& NewClassName, const FStri
 
 		TArray<FAnalyticsEventAttribute> EventAttributes;
 		EventAttributes.Add(FAnalyticsEventAttribute(TEXT("ParentClass"), ParentClassName.IsEmpty() ? TEXT("None") : ParentClassName));
-		EventAttributes.Add(FAnalyticsEventAttribute(TEXT("Outcome"), bAddCodeSuccessful ? TEXT("Successful") : TEXT("Failed")));
+		EventAttributes.Add(FAnalyticsEventAttribute(TEXT("Outcome"), Result == EAddCodeToProjectResult::Succeeded ? TEXT("Successful") : TEXT("Failed")));
+		EventAttributes.Add(FAnalyticsEventAttribute(TEXT("FailureReason"), OutFailReason.ToString()));
 
 		FEngineAnalytics::GetProvider().RecordEvent( TEXT( "Editor.AddCodeToProject.CodeAdded" ), EventAttributes );
 	}
 
-	return bAddCodeSuccessful;
+	return Result;
 }
 
 UTemplateProjectDefs* GameProjectUtils::LoadTemplateDefs(const FString& ProjectDirectory)
@@ -2958,12 +2959,12 @@ TArray<FString> GameProjectUtils::GetRequiredAdditionalDependencies(const FNewCl
 	return Out;
 }
 
-bool GameProjectUtils::AddCodeToProject_Internal(const FString& NewClassName, const FString& NewClassPath, const FModuleContextInfo& ModuleInfo, const FNewClassInfo ParentClassInfo, const TSet<FString>& DisallowedHeaderNames, FString& OutHeaderFilePath, FString& OutCppFilePath, FText& OutFailReason)
+GameProjectUtils::EAddCodeToProjectResult GameProjectUtils::AddCodeToProject_Internal(const FString& NewClassName, const FString& NewClassPath, const FModuleContextInfo& ModuleInfo, const FNewClassInfo ParentClassInfo, const TSet<FString>& DisallowedHeaderNames, FString& OutHeaderFilePath, FString& OutCppFilePath, FText& OutFailReason)
 {
 	if ( !ParentClassInfo.IsSet() )
 	{
 		OutFailReason = LOCTEXT("NoParentClass", "You must specify a parent class");
-		return false;
+		return EAddCodeToProjectResult::InvalidInput;
 	}
 
 	const FString CleanClassName = ParentClassInfo.GetCleanClassName(NewClassName);
@@ -2971,20 +2972,20 @@ bool GameProjectUtils::AddCodeToProject_Internal(const FString& NewClassName, co
 
 	if (!IsValidClassNameForCreation(FinalClassName, ModuleInfo, DisallowedHeaderNames, OutFailReason))
 	{
-		return false;
+		return EAddCodeToProjectResult::InvalidInput;
 	}
 
 	if ( !FApp::HasGameName() )
 	{
 		OutFailReason = LOCTEXT("AddCodeToProject_NoGameName", "You can not add code because you have not loaded a project.");
-		return false;
+		return EAddCodeToProjectResult::FailedToAddCode;
 	}
 
 	FString NewHeaderPath;
 	FString NewCppPath;
 	if ( !CalculateSourcePaths(NewClassPath, ModuleInfo, NewHeaderPath, NewCppPath, &OutFailReason) )
 	{
-		return false;
+		return EAddCodeToProjectResult::FailedToAddCode;
 	}
 
 	FScopedSlowTask SlowTask( 7, LOCTEXT( "AddingCodeToProject", "Adding code to project..." ) );
@@ -3018,7 +3019,7 @@ bool GameProjectUtils::AddCodeToProject_Internal(const FString& NewClassName, co
 		else
 		{
 			DeleteCreatedFiles(SourceDir, CreatedFiles);
-			return false;
+			return EAddCodeToProjectResult::FailedToAddCode;
 		}
 	}
 
@@ -3050,7 +3051,7 @@ bool GameProjectUtils::AddCodeToProject_Internal(const FString& NewClassName, co
 		else
 		{
 			DeleteCreatedFiles(NewHeaderPath, CreatedFiles);
-			return false;
+			return EAddCodeToProjectResult::FailedToAddCode;
 		}
 	}
 
@@ -3067,7 +3068,7 @@ bool GameProjectUtils::AddCodeToProject_Internal(const FString& NewClassName, co
 		else
 		{
 			DeleteCreatedFiles(NewCppPath, CreatedFiles);
-			return false;
+			return EAddCodeToProjectResult::FailedToAddCode;
 		}
 	}
 
@@ -3099,7 +3100,7 @@ bool GameProjectUtils::AddCodeToProject_Internal(const FString& NewClassName, co
 		if ( !FDesktopPlatformModule::Get()->GenerateProjectFiles(FPaths::RootDir(), FPaths::GetProjectFilePath(), GWarn) )
 		{
 			OutFailReason = LOCTEXT("FailedToGenerateProjectFiles", "Failed to generate project files.");
-			return false;
+			return EAddCodeToProjectResult::FailedToHotReload;
 		}
 	}
 
@@ -3130,7 +3131,7 @@ bool GameProjectUtils::AddCodeToProject_Internal(const FString& NewClassName, co
 		if (!HotReloadSupport.RecompileModule(*GameModuleName, bReloadAfterCompiling, *GWarn, bFailIfGeneratedCodeChanges, bForceCodeProject))
 		{
 			OutFailReason = LOCTEXT("FailedToCompileNewGameModule", "Failed to compile newly created game module.");
-			return false;
+			return EAddCodeToProjectResult::FailedToHotReload;
 		}
 
 		// Notify that we've created a brand new module
@@ -3168,7 +3169,7 @@ bool GameProjectUtils::AddCodeToProject_Internal(const FString& NewClassName, co
 				if( CompilationResult != ECompilationResult::Succeeded && CompilationResult != ECompilationResult::UpToDate )
 				{
 					OutFailReason = FText::Format(LOCTEXT("FailedToHotReloadModuleFmt", "Failed to automatically hot reload the '{0}' module."), FText::FromString(ModuleInfo.ModuleName));
-					return false;
+					return EAddCodeToProjectResult::FailedToHotReload;
 				}
 			}
 			else
@@ -3180,13 +3181,13 @@ bool GameProjectUtils::AddCodeToProject_Internal(const FString& NewClassName, co
 				if (!HotReloadSupport.RecompileModule(ModuleFName, bReloadAfterRecompile, *GWarn, bFailIfGeneratedCodeChanges, bForceCodeProject))
 				{
 					OutFailReason = FText::Format(LOCTEXT("FailedToCompileModuleFmt", "Failed to automatically compile the '{0}' module."), FText::FromString(ModuleInfo.ModuleName));
-					return false;
+					return EAddCodeToProjectResult::FailedToHotReload;
 				}
 			}
 		}
 	}
 
-	return true;
+	return EAddCodeToProjectResult::Succeeded;
 }
 
 bool GameProjectUtils::FindSourceFileInProject(const FString& InFilename, const FString& InSearchPath, FString& OutPath)
diff --git a/Engine/Source/Editor/GameProjectGeneration/Private/GameProjectUtils.h b/Engine/Source/Editor/GameProjectGeneration/Private/GameProjectUtils.h
--- a/Engine/Source/Editor/GameProjectGeneration/Private/GameProjectUtils.h
+++ b/Engine/Source/Editor/GameProjectGeneration/Private/GameProjectUtils.h
@@ -50,6 +50,22 @@ public:
 		Classes,
 	};
 
+	/** Used as a function return result when adding new code to the project */
+	enum class EAddCodeToProjectResult : uint8
+	{
+		/** Function has successfully added the code and hot-reloaded the required module(s) */
+		Succeeded,
+
+		/** There were errors with the input given to the function */
+		InvalidInput,
+
+		/** There were errors when adding the new source files */
+		FailedToAddCode,
+		 
+		/** There were errors when hot-reloading the new module */
+		FailedToHotReload,
+	};
+
 	/** Used as a function return result when a project is duplicated when upgrading project's version in Convert project dialog - Open a copy */
 	enum class EProjectDuplicateResult : uint8
 	{
@@ -107,8 +123,8 @@ public:
 	/** Returns true if the specified class is a valid base class for any of the given modules */
 	static bool IsValidBaseClassForCreation(const UClass* InClass, const TArray<FModuleContextInfo>& InModuleInfoArray);
 
-	/** Adds new source code to the project. When returning true, OutSyncFileAndLineNumber will be the the preferred target file to sync in the users code editing IDE, formatted for use with GenericApplication::GotoLineInSource */
-	static bool AddCodeToProject(const FString& NewClassName, const FString& NewClassPath, const FModuleContextInfo& ModuleInfo, const FNewClassInfo ParentClassInfo, const TSet<FString>& DisallowedHeaderNames, FString& OutHeaderFilePath, FString& OutCppFilePath, FText& OutFailReason);
+	/** Adds new source code to the project. When returning Succeeded or FailedToHotReload, OutSyncFileAndLineNumber will be the the preferred target file to sync in the users code editing IDE, formatted for use with GenericApplication::GotoLineInSource */
+	static EAddCodeToProjectResult AddCodeToProject(const FString& NewClassName, const FString& NewClassPath, const FModuleContextInfo& ModuleInfo, const FNewClassInfo ParentClassInfo, const TSet<FString>& DisallowedHeaderNames, FString& OutHeaderFilePath, FString& OutCppFilePath, FText& OutFailReason);
 
 	/** Loads a template project definitions object from the TemplateDefs.ini file in the specified project */
 	static UTemplateProjectDefs* LoadTemplateDefs(const FString& ProjectDirectory);
@@ -369,7 +385,7 @@ private:
 	static bool CheckoutGameProjectFile(const FString& ProjectFilename, FText& OutFailReason);
 
 	/** Internal handler for AddCodeToProject*/
-	static bool AddCodeToProject_Internal(const FString& NewClassName, const FString& NewClassPath, const FModuleContextInfo& ModuleInfo, const FNewClassInfo ParentClassInfo, const TSet<FString>& DisallowedHeaderNames, FString& OutHeaderFilePath, FString& OutCppFilePath, FText& OutFailReason);
+	static EAddCodeToProjectResult AddCodeToProject_Internal(const FString& NewClassName, const FString& NewClassPath, const FModuleContextInfo& ModuleInfo, const FNewClassInfo ParentClassInfo, const TSet<FString>& DisallowedHeaderNames, FString& OutHeaderFilePath, FString& OutCppFilePath, FText& OutFailReason);
 
 	/** Internal handler for IsValidBaseClassForCreation */
 	DECLARE_DELEGATE_RetVal_OneParam(bool, FDoesClassNeedAPIExportCallback, const FString& /*ClassModuleName*/);
diff --git a/Engine/Source/Editor/GameProjectGeneration/Private/SNewClassDialog.cpp b/Engine/Source/Editor/GameProjectGeneration/Private/SNewClassDialog.cpp
--- a/Engine/Source/Editor/GameProjectGeneration/Private/SNewClassDialog.cpp
+++ b/Engine/Source/Editor/GameProjectGeneration/Private/SNewClassDialog.cpp
@@ -1156,7 +1156,8 @@ void SNewClassDialog::FinishClicked()
 
 		FText FailReason;
 		const TSet<FString>& DisallowedHeaderNames = FSourceCodeNavigation::GetSourceFileDatabase().GetDisallowedHeaderNames();
-		if (GameProjectUtils::AddCodeToProject(NewClassName, NewClassPath, *SelectedModuleInfo, ParentClassInfo, DisallowedHeaderNames, HeaderFilePath, CppFilePath, FailReason))
+		const GameProjectUtils::EAddCodeToProjectResult AddCodeResult = GameProjectUtils::AddCodeToProject(NewClassName, NewClassPath, *SelectedModuleInfo, ParentClassInfo, DisallowedHeaderNames, HeaderFilePath, CppFilePath, FailReason);
+		if (AddCodeResult == GameProjectUtils::EAddCodeToProjectResult::Succeeded)
 		{
 			OnAddedToProject.ExecuteIfBound( NewClassName, NewClassPath, SelectedModuleInfo->ModuleName );
 
@@ -1177,7 +1178,7 @@ void SNewClassDialog::FinishClicked()
 				{
 					// Code successfully added, notify the user. We are either running on a platform that does not support source access or a file was not given so don't ask about editing the file
 					const FText Message = FText::Format( 
-						LOCTEXT("AddCodeSuccessWithHotReload", "Successfully added class {0}, however you must recompile {1} before it will appear in the Content Browser.")
+						LOCTEXT("AddCodeSuccessWithHotReload", "Successfully added class '{0}', however you must recompile the '{1}' module before it will appear in the Content Browser.")
 						, FText::FromString(NewClassName), FText::FromString(SelectedModuleInfo->ModuleName) );
 					FMessageDialog::Open(EAppMsgType::Ok, Message);
 				}
@@ -1198,7 +1199,7 @@ void SNewClassDialog::FinishClicked()
 				{
 					// Code successfully added, notify the user and ask about opening the IDE now
 					const FText Message = FText::Format( 
-						LOCTEXT("AddCodeSuccessWithHotReloadAndSync", "Successfully added class {0}, however you must recompile {1} before it will appear in the Content Browser.\n\nWould you like to edit the code now?")
+						LOCTEXT("AddCodeSuccessWithHotReloadAndSync", "Successfully added class '{0}', however you must recompile the '{1}' module before it will appear in the Content Browser.\n\nWould you like to edit the code now?")
 						, FText::FromString(NewClassName), FText::FromString(SelectedModuleInfo->ModuleName) );
 					bEditSourceFilesNow = ( FMessageDialog::Open( EAppMsgType::YesNo, Message ) == EAppReturnType::Yes );
 				}
@@ -1230,15 +1231,31 @@ void SNewClassDialog::FinishClicked()
 			// Successfully created the code and potentially opened the IDE. Close the dialog.
 			CloseContainingWindow();
 		}
-		else
+		else if (AddCodeResult == GameProjectUtils::EAddCodeToProjectResult::FailedToHotReload)
 		{
-			// @todo show fail reason in error label
-			// Failed to add code
-			const FText Message = FText::Format( LOCTEXT("AddCodeFailed_Native", "Failed to add or compile class {0}. {1}\n\nWould you like to open the Output Log to see more details?"), FText::FromString(NewClassName), FailReason );
+			OnAddedToProject.ExecuteIfBound( NewClassName, NewClassPath, SelectedModuleInfo->ModuleName );
+
+			// Prevent periodic validity checks. This is to prevent a brief error message about the class already existing while you are exiting.
+			bPreventPeriodicValidityChecksUntilNextChange = true;
+
+			// Failed to compile new code
+			const FText Message = FText::Format(
+				LOCTEXT("AddCodeFailed_HotReloadFailed", "Successfully added class '{0}', however you must recompile the '{1}' module before it will appear in the Content Browser. {2}\n\nWould you like to open the Output Log to see more details?")
+				, FText::FromString(NewClassName), FText::FromString(SelectedModuleInfo->ModuleName), FailReason );
 			if( FMessageDialog::Open(EAppMsgType::YesNo, Message) == EAppReturnType::Yes )
 			{
 				FGlobalTabmanager::Get()->InvokeTab(FName("OutputLog"));
 			}
+
+			// We did manage to add the code itself, so we can close the dialog.
+			CloseContainingWindow();
+		}
+		else
+		{
+			// @todo show fail reason in error label
+			// Failed to add code
+			const FText Message = FText::Format( LOCTEXT("AddCodeFailed_AddCodeFailed", "Failed to add class '{0}'. {1}"), FText::FromString(NewClassName), FailReason );
+			FMessageDialog::Open(EAppMsgType::Ok, Message);
 		}
 	}
 }
