Editor: Switched SNewClassDialog to use colorized buttons

diff --git a/Engine/Source/Editor/GameProjectGeneration/Private/SNewClassDialog.cpp b/Engine/Source/Editor/GameProjectGeneration/Private/SNewClassDialog.cpp
--- a/Engine/Source/Editor/GameProjectGeneration/Private/SNewClassDialog.cpp
+++ b/Engine/Source/Editor/GameProjectGeneration/Private/SNewClassDialog.cpp
@@ -235,6 +235,13 @@ void SNewClassDialog::Construct( const FArguments& InArgs )
 			[
 				SAssignNew( MainWizard, SWizard)
 				.ShowPageList(false)
+
+				.ButtonStyle(FEditorStyle::Get(), "FlatButton.Default")
+				.CancelButtonStyle(FEditorStyle::Get(), "FlatButton.Warning")
+				.FinishButtonStyle(FEditorStyle::Get(), "FlatButton.Success")
+				.ButtonTextStyle(FEditorStyle::Get(), "LargeText")
+				.ForegroundColor(FEditorStyle::Get().GetSlateColor("WhiteBrush"))
+
 				.CanFinish(this, &SNewClassDialog::CanFinish)
 				.FinishButtonText( ClassDomain == EClassDomain::Native ? LOCTEXT("FinishButtonText_Native", "Create Class") : LOCTEXT("FinishButtonText_Blueprint", "Create Blueprint Class") )
 				.FinishButtonToolTip (
