Fix orange cancel buttons in the new class and new project wizard dialogs

diff --git a/Engine/Source/Editor/GameProjectGeneration/Private/SNewClassDialog.cpp b/Engine/Source/Editor/GameProjectGeneration/Private/SNewClassDialog.cpp
--- a/Engine/Source/Editor/GameProjectGeneration/Private/SNewClassDialog.cpp
+++ b/Engine/Source/Editor/GameProjectGeneration/Private/SNewClassDialog.cpp
@@ -237,7 +237,7 @@ void SNewClassDialog::Construct( const FArguments& InArgs )
 				.ShowPageList(false)
 
 				.ButtonStyle(FEditorStyle::Get(), "FlatButton.Default")
-				.CancelButtonStyle(FEditorStyle::Get(), "FlatButton.Warning")
+				.CancelButtonStyle(FEditorStyle::Get(), "FlatButton.Default")
 				.FinishButtonStyle(FEditorStyle::Get(), "FlatButton.Success")
 				.ButtonTextStyle(FEditorStyle::Get(), "LargeText")
 				.ForegroundColor(FEditorStyle::Get().GetSlateColor("WhiteBrush"))
diff --git a/Engine/Source/Editor/GameProjectGeneration/Private/SNewProjectWizard.cpp b/Engine/Source/Editor/GameProjectGeneration/Private/SNewProjectWizard.cpp
--- a/Engine/Source/Editor/GameProjectGeneration/Private/SNewProjectWizard.cpp
+++ b/Engine/Source/Editor/GameProjectGeneration/Private/SNewProjectWizard.cpp
@@ -358,7 +358,7 @@ void SNewProjectWizard::Construct( const FArguments& InArgs )
 			[
 				SAssignNew(MainWizard, SWizard)
 				.ButtonStyle(FEditorStyle::Get(), "FlatButton.Default")
-				.CancelButtonStyle(FEditorStyle::Get(), "FlatButton.Warning")
+				.CancelButtonStyle(FEditorStyle::Get(), "FlatButton.Default")
 				.FinishButtonStyle(FEditorStyle::Get(), "FlatButton.Success")
 				.ButtonTextStyle(FEditorStyle::Get(), "LargeText")
 				.ForegroundColor(FEditorStyle::Get().GetSlateColor("WhiteBrush"))
