Fix empty projects not having a ProjectID set by default.

diff --git a/Engine/Source/Editor/GameProjectGeneration/Private/GameProjectUtils.cpp b/Engine/Source/Editor/GameProjectGeneration/Private/GameProjectUtils.cpp
--- a/Engine/Source/Editor/GameProjectGeneration/Private/GameProjectUtils.cpp
+++ b/Engine/Source/Editor/GameProjectGeneration/Private/GameProjectUtils.cpp
@@ -1764,6 +1764,23 @@ bool GameProjectUtils::GenerateConfigFiles(const FProjectInformation& InProjectI
 		}
 	}
 
+	// DefaultGame.ini
+	{
+		const FString DefaultGameIniFilename = ProjectConfigPath / TEXT("DefaultGame.ini");
+		FString FileContents;
+		FileContents += TEXT("[/Script/EngineSettings.GeneralProjectSettings]") LINE_TERMINATOR;
+		FileContents += TEXT("ProjectID=") + FGuid::NewGuid().ToString() + LINE_TERMINATOR;
+
+		if (WriteOutputFile(DefaultGameIniFilename, FileContents, OutFailReason))
+		{
+			OutCreatedFiles.Add(DefaultGameIniFilename);
+		}
+		else
+		{
+			return false;
+		}
+	}
+
 	return true;
 }
 
