Missing header.

diff --git a/Engine/Source/Editor/GameProjectGeneration/Public/GameProjectUtils.h b/Engine/Source/Editor/GameProjectGeneration/Public/GameProjectUtils.h
--- a/Engine/Source/Editor/GameProjectGeneration/Public/GameProjectUtils.h
+++ b/Engine/Source/Editor/GameProjectGeneration/Public/GameProjectUtils.h
@@ -3,6 +3,7 @@
 #pragma once
 
 #include "HardwareTargetingModule.h"
+#include "ProjectDescriptor.h"
 
 struct FModuleContextInfo;
 struct FAddToProjectConfig;
