Don't clobber the DefaultGame.ini file when adding feature packs to a project; it wipes out the project ID.

diff --git a/Engine/Source/Editor/GameProjectGeneration/Private/GameProjectUtils.cpp b/Engine/Source/Editor/GameProjectGeneration/Private/GameProjectUtils.cpp
--- a/Engine/Source/Editor/GameProjectGeneration/Private/GameProjectUtils.cpp
+++ b/Engine/Source/Editor/GameProjectGeneration/Private/GameProjectUtils.cpp
@@ -1025,7 +1025,11 @@ bool GameProjectUtils::GenerateProjectFromScratch(const FProjectInformation& InP
 	}
 
 	// Insert any required feature packs (EG starter content) into ini file. These will be imported automatically when the editor is first run
-	InsertFeaturePacksIntoINIFile(InProjectInfo, OutFailReason);
+	if(!InsertFeaturePacksIntoINIFile(InProjectInfo, OutFailReason))
+	{
+		DeleteCreatedFiles(NewProjectFolder, CreatedFiles);
+		return false;
+	}
 	
 	// Make the Content folder
 	const FString ContentFolder = NewProjectFolder / TEXT("Content");
@@ -1462,7 +1466,11 @@ bool GameProjectUtils::CreateProjectFromTemplate(const FProjectInformation& InPr
 	}
 
 	// Insert any required feature packs (EG starter content) into ini file. These will be imported automatically when the editor is first run
-	InsertFeaturePacksIntoINIFile(InProjectInfo, OutFailReason);
+	if(!InsertFeaturePacksIntoINIFile(InProjectInfo, OutFailReason))
+	{
+		DeleteCreatedFiles(DestFolder, CreatedFiles);
+		return false;
+	}
 	
 	SlowTask.EnterProgressFrame();
 
@@ -3250,7 +3258,6 @@ void GameProjectUtils::HarvestCursorSyncLocation( FString& FinalOutput, FString&
 
 bool GameProjectUtils::InsertFeaturePacksIntoINIFile(const FProjectInformation& InProjectInfo, FText& OutFailReason)
 {
-	bool bSuccessfullyProcessed = false;
 	const FString ProjectName = FPaths::GetBaseFilename(InProjectInfo.ProjectFilename);
 	const FString TemplateName = FPaths::GetBaseFilename(InProjectInfo.TemplateFile);
 	const FString SrcFolder = FPaths::GetPath(InProjectInfo.TemplateFile);
@@ -3286,25 +3293,33 @@ bool GameProjectUtils::InsertFeaturePacksIntoINIFile(const FProjectInformation&
 		}
 	}
 
-	FString FileOutput;
 	if (PackList.Num() != 0)
 	{
-		FileOutput = TEXT("[StartupActions]");
+		FString FileOutput;
+		if(FPaths::FileExists(IniFilename) && !FFileHelper::LoadFileToString(FileOutput, *IniFilename))
+		{
+			OutFailReason = LOCTEXT("FailedToReadIni", "Could not read INI file to insert feature packs");
+			return false;
+		}
+
+		FileOutput += LINE_TERMINATOR;
+		FileOutput += TEXT("[StartupActions]");
 		FileOutput += LINE_TERMINATOR;
 		FileOutput += TEXT("bAddPacks=True");
 		FileOutput += LINE_TERMINATOR;
 		for (int32 iLine = 0; iLine < PackList.Num(); ++iLine)
 		{
 			FileOutput += PackList[iLine] + LINE_TERMINATOR;
-
 		}
-		if (FFileHelper::SaveStringToFile(FileOutput, *IniFilename))
+
+		if (!FFileHelper::SaveStringToFile(FileOutput, *IniFilename))
 		{
-			bSuccessfullyProcessed = true;
+			OutFailReason = LOCTEXT("FailedToWriteIni", "Could not write INI file to insert feature packs");
+			return false;
 		}
 	}
 
-	return bSuccessfullyProcessed;
+	return true;
 }
 
 #undef LOCTEXT_NAMESPACE
