fixed issue with not being able to tell if a content-based project needs to be built as code within the editor
UE-11636

diff --git a/Engine/Source/Developer/Android/AndroidTargetPlatform/Private/AndroidTargetPlatform.h b/Engine/Source/Developer/Android/AndroidTargetPlatform/Private/AndroidTargetPlatform.h
--- a/Engine/Source/Developer/Android/AndroidTargetPlatform/Private/AndroidTargetPlatform.h
+++ b/Engine/Source/Developer/Android/AndroidTargetPlatform/Private/AndroidTargetPlatform.h
@@ -145,6 +145,13 @@ public:
 
 	virtual FText GetVariantTitle() const override;
 
+	virtual void GetBuildProjectSettingKeys(FString& OutSection, TArray<FString>& InBoolKeys, TArray<FString>& InIntKeys, TArray<FString>& InStringKeys) const override
+	{
+		OutSection = TEXT("/Script/AndroidRuntimeSettings.AndroidRuntimeSettings");
+		InBoolKeys.Add(TEXT("bBuildForArmV7")); InBoolKeys.Add(TEXT("bBuildForArm64")); InBoolKeys.Add(TEXT("bBuildForX86"));
+		InBoolKeys.Add(TEXT("bBuildForX8664")); InBoolKeys.Add(TEXT("bBuildForES2")); InBoolKeys.Add(TEXT("bBuildForES31"));
+	}
+
 	DECLARE_DERIVED_EVENT(FAndroidTargetPlatform, ITargetPlatform::FOnTargetDeviceDiscovered, FOnTargetDeviceDiscovered);
 	virtual FOnTargetDeviceDiscovered& OnDeviceDiscovered( ) override
 	{
diff --git a/Engine/Source/Developer/TargetPlatform/Public/Common/TargetPlatformBase.h b/Engine/Source/Developer/TargetPlatform/Public/Common/TargetPlatformBase.h
--- a/Engine/Source/Developer/TargetPlatform/Public/Common/TargetPlatformBase.h
+++ b/Engine/Source/Developer/TargetPlatform/Public/Common/TargetPlatformBase.h
@@ -228,6 +228,11 @@ public:
 		return false;
 	}
 
+	virtual void GetBuildProjectSettingKeys(FString& OutSection, TArray<FString>& InBoolKeys, TArray<FString>& InIntKeys, TArray<FString>& InStringKeys) const override
+	{
+		// do nothing in the base class
+	}
+
 protected:
 
 	FTargetPlatformBase(const PlatformInfo::FPlatformInfo *const InPlatformInfo)
diff --git a/Engine/Source/Developer/TargetPlatform/Public/Interfaces/ITargetPlatform.h b/Engine/Source/Developer/TargetPlatform/Public/Interfaces/ITargetPlatform.h
--- a/Engine/Source/Developer/TargetPlatform/Public/Interfaces/ITargetPlatform.h
+++ b/Engine/Source/Developer/TargetPlatform/Public/Interfaces/ITargetPlatform.h
@@ -370,6 +370,11 @@ public:
 	 */
 	virtual bool SendLowerCaseFilePaths() const = 0;
 
+	/**
+	* Project settings to check to determine if a build should occurr
+	*/
+	virtual void GetBuildProjectSettingKeys(FString& OutSection, TArray<FString>& InBoolKeys, TArray<FString>& InIntKeys, TArray<FString>& InStringKeys) const = 0;
+
 public:
 
 	/**
diff --git a/Engine/Source/Developer/iOS/IOSTargetPlatform/Private/IOSTargetPlatform.h b/Engine/Source/Developer/iOS/IOSTargetPlatform/Private/IOSTargetPlatform.h
--- a/Engine/Source/Developer/iOS/IOSTargetPlatform/Private/IOSTargetPlatform.h
+++ b/Engine/Source/Developer/iOS/IOSTargetPlatform/Private/IOSTargetPlatform.h
@@ -105,6 +105,14 @@ public:
 	virtual FName GetWaveFormat( const class USoundWave* Wave ) const override;
 #endif // WITH_ENGINE
 
+	virtual void GetBuildProjectSettingKeys(FString& OutSection, TArray<FString>& InBoolKeys, TArray<FString>& InIntKeys, TArray<FString>& InStringKeys) const override
+	{
+		OutSection = TEXT("/Script/IOSRuntimeSettings.IOSRuntimeSettings");
+		InBoolKeys.Add(TEXT("bDevForArmV7")); InBoolKeys.Add(TEXT("bDevForArm64")); InBoolKeys.Add(TEXT("bDevForArmV7S"));
+		InBoolKeys.Add(TEXT("bShipForArmV7")); InBoolKeys.Add(TEXT("bShipForArm64")); InBoolKeys.Add(TEXT("bShipForArmV7S"));
+		InBoolKeys.Add(TEXT("bGenerateSYMFile"));
+		InStringKeys.Add(TEXT("MinimumiOSVersion"));
+	}
 
 	DECLARE_DERIVED_EVENT(FIOSTargetPlatform, ITargetPlatform::FOnTargetDeviceDiscovered, FOnTargetDeviceDiscovered);
 	virtual FOnTargetDeviceDiscovered& OnDeviceDiscovered( ) override
diff --git a/Engine/Source/Editor/GameProjectGeneration/Private/GameProjectGenerationModule.cpp b/Engine/Source/Editor/GameProjectGeneration/Private/GameProjectGenerationModule.cpp
--- a/Engine/Source/Editor/GameProjectGeneration/Private/GameProjectGenerationModule.cpp
+++ b/Engine/Source/Editor/GameProjectGeneration/Private/GameProjectGenerationModule.cpp
@@ -102,6 +102,11 @@ bool FGameProjectGenerationModule::ProjectHasCodeFiles()
 	return GameProjectUtils::ProjectHasCodeFiles();
 }
 
+bool FGameProjectGenerationModule::ProjectRequiresBuild(const FName InPlatformName)
+{
+	return GameProjectUtils::ProjectRequiresBuild(InPlatformName);
+}
+
 FString FGameProjectGenerationModule::DetermineModuleIncludePath(const FModuleContextInfo& ModuleInfo, const FString& FileRelativeTo)
 {
 	return GameProjectUtils::DetermineModuleIncludePath(ModuleInfo, FileRelativeTo);
diff --git a/Engine/Source/Editor/GameProjectGeneration/Private/GameProjectUtils.cpp b/Engine/Source/Editor/GameProjectGeneration/Private/GameProjectUtils.cpp
--- a/Engine/Source/Editor/GameProjectGeneration/Private/GameProjectUtils.cpp
+++ b/Engine/Source/Editor/GameProjectGeneration/Private/GameProjectUtils.cpp
@@ -31,6 +31,8 @@
 
 #include "SOutputLogDialog.h"
 
+#include "PlatformInfo.h"
+
 #define LOCTEXT_NAMESPACE "GameProjectUtils"
 
 #define MAX_PROJECT_PATH_BUFFER_SPACE 130 // Leave a reasonable buffer of additional characters to account for files created in the content directory during or after project generation
@@ -2955,6 +2957,107 @@ bool GameProjectUtils::ProjectHasCodeFiles()
 	return GameProjectUtils::GetProjectCodeFileCount() > 0;
 }
 
+bool GameProjectUtils::ProjectRequiresBuild(const FName InPlatformInfoName)
+{
+	//  early out on projects with code files
+	if (ProjectHasCodeFiles())
+	{
+		return true;
+	}
+
+	bool bRequiresBuild = false;
+	// check to see if the default build settings have changed
+	bRequiresBuild |= !HasDefaultBuildSettings(InPlatformInfoName);
+
+	// check to see if any plugins beyond the defaults have been enabled
+	bRequiresBuild |= IProjectManager::Get().IsNonDefaultPluginEnabled();
+
+	return bRequiresBuild;
+}
+
+bool GameProjectUtils::DoProjectSettingsMatchDefault(const FString& InPlatformName, const FString& InSection, const TArray<FString>* InBoolKeys, const TArray<FString>* InIntKeys, const TArray<FString>* InStringKeys)
+{
+	FConfigFile ProjIni;
+	FConfigFile DefaultIni;
+	FConfigCacheIni::LoadLocalIniFile(ProjIni, TEXT("Engine"), true, *InPlatformName, true);
+	FConfigCacheIni::LoadLocalIniFile(DefaultIni, TEXT("Engine"), true, NULL, true);
+
+	if (InBoolKeys != NULL)
+	{
+		for (int Index = 0; Index < InBoolKeys->Num(); ++Index)
+		{
+			FString Default(TEXT("False")), Project(TEXT("False"));
+			DefaultIni.GetString(*InSection, *((*InBoolKeys)[Index]), Default);
+			ProjIni.GetString(*InSection, *((*InBoolKeys)[Index]), Project);
+			if (!Default.Compare(Project, ESearchCase::IgnoreCase))
+			{
+				return false;
+			}
+		}
+	}
+
+	if (InIntKeys != NULL)
+	{
+		for (int Index = 0; Index < InIntKeys->Num(); ++Index)
+		{
+			int64 Default(0), Project(0);
+			DefaultIni.GetInt64(*InSection, *((*InIntKeys)[Index]), Default);
+			ProjIni.GetInt64(*InSection, *((*InIntKeys)[Index]), Project);
+			if (Default != Project)
+			{
+				return false;
+			}
+		}
+	}
+
+	if (InStringKeys != NULL)
+	{
+		for (int Index = 0; Index < InStringKeys->Num(); ++Index)
+		{
+			FString Default(TEXT("False")), Project(TEXT("False"));
+			DefaultIni.GetString(*InSection, *((*InStringKeys)[Index]), Default);
+			ProjIni.GetString(*InSection, *((*InStringKeys)[Index]), Project);
+			if (!Default.Compare(Project, ESearchCase::IgnoreCase))
+			{
+				return false;
+			}
+		}
+	}
+
+	return true;
+}
+
+bool GameProjectUtils::HasDefaultBuildSettings(const FName InPlatformInfoName)
+{
+	// first check default build settings for all platforms
+	TArray<FString> BoolKeys, IntKeys, StringKeys, BuildKeys;
+	BuildKeys.Add(TEXT("bCompileApex")); BuildKeys.Add(TEXT("bCompileBox2D")); BuildKeys.Add(TEXT("bCompileICU"));
+	BuildKeys.Add(TEXT("bCompileSimplygon")); BuildKeys.Add(TEXT("bCompi8leLeanAndMeanUE")); BuildKeys.Add(TEXT("bIncludeADO"));
+	BuildKeys.Add(TEXT("bCompileRecast")); BuildKeys.Add(TEXT("bCompileSpeedTree")); BuildKeys.Add(TEXT("bCompileWithPluginSupport"));
+	BuildKeys.Add(TEXT("bCompilePhysXVehicle")); BuildKeys.Add(TEXT("bCompileFreeType")); BuildKeys.Add(TEXT("bCompileForSize"));
+	BuildKeys.Add(TEXT("bCompileCEF3"));
+
+	const PlatformInfo::FPlatformInfo* const PlatInfo = PlatformInfo::FindPlatformInfo(InPlatformInfoName);
+	check(PlatInfo);
+
+	if (!DoProjectSettingsMatchDefault(PlatInfo->TargetPlatformName.ToString(), TEXT("/Script/BuildSettings.BuildSettings"), &BuildKeys))
+	{
+		return false;
+	}
+
+	if (PlatInfo->SDKStatus == PlatformInfo::EPlatformSDKStatus::Installed)
+	{
+		const ITargetPlatform* const Platform = GetTargetPlatformManager()->FindTargetPlatform(PlatInfo->TargetPlatformName.ToString());
+		if (Platform)
+		{
+			FString PlatformSection;
+			Platform->GetBuildProjectSettingKeys(PlatformSection, BoolKeys, IntKeys, StringKeys);
+			return DoProjectSettingsMatchDefault(PlatInfo->TargetPlatformName.ToString(), PlatformSection, &BoolKeys, &IntKeys, &StringKeys);
+		}
+	}
+	return true;
+}
+
 TArray<FString> GameProjectUtils::GetRequiredAdditionalDependencies(const FNewClassInfo& ClassInfo)
 {
 	TArray<FString> Out;
diff --git a/Engine/Source/Editor/GameProjectGeneration/Public/GameProjectGenerationModule.h b/Engine/Source/Editor/GameProjectGeneration/Public/GameProjectGenerationModule.h
--- a/Engine/Source/Editor/GameProjectGeneration/Public/GameProjectGenerationModule.h
+++ b/Engine/Source/Editor/GameProjectGeneration/Public/GameProjectGenerationModule.h
@@ -125,6 +125,9 @@ public:
 	/** Clear the list of supported target platforms */
 	virtual void ClearSupportedTargetPlatforms();
 
+	/** Check to see if the project requires a code build */
+	virtual bool ProjectRequiresBuild(const FName InPlatformInfoName);
+
 public:
 
 	/** (Un)register a new type of template category to be shown on the new project page */
diff --git a/Engine/Source/Editor/GameProjectGeneration/Public/GameProjectUtils.h b/Engine/Source/Editor/GameProjectGeneration/Public/GameProjectUtils.h
--- a/Engine/Source/Editor/GameProjectGeneration/Public/GameProjectUtils.h
+++ b/Engine/Source/Editor/GameProjectGeneration/Public/GameProjectUtils.h
@@ -234,6 +234,9 @@ public:
 	/** Returns a list of #include lines formed from InList */
 	static FString MakeIncludeList(const TArray<FString>& InList);
 
+	/** Returns true if the currently loaded project requires a code build */
+	static bool ProjectRequiresBuild(const FName InPlatformInfoName);
+
 private:
 
 	static FString GetHardwareConfigString(const FProjectInformation& InProjectInfo);
@@ -431,6 +434,10 @@ private:
 	 */
 	static bool UpdateRequiredAdditionalDependencies(FProjectDescriptor& Descriptor, TArray<FString>& RequiredDependencies, const FString& ModuleName);
 
+	/** checks the project ini file against the default ini file to determine whether or not the build settings have changed from default */
+	static bool HasDefaultBuildSettings(const FName InPlatformInfoName);
+	static bool DoProjectSettingsMatchDefault(const FString& InPlatformnName, const FString& InSection, const TArray<FString>* InBoolKeys, const TArray<FString>* InIntKeys = NULL, const TArray<FString>* InStringKeys = NULL);
+
 private:
 	/**
 	 * Updates the projects and modifies FProjectDescriptor accordingly to given modifier.
diff --git a/Engine/Source/Editor/MainFrame/Private/Frame/MainFrameActions.cpp b/Engine/Source/Editor/MainFrame/Private/Frame/MainFrameActions.cpp
--- a/Engine/Source/Editor/MainFrame/Private/Frame/MainFrameActions.cpp
+++ b/Engine/Source/Editor/MainFrame/Private/Frame/MainFrameActions.cpp
@@ -369,7 +369,7 @@ void FMainFrameActionCallbacks::PackageProject( const FName InPlatformInfoName )
 {
 	// does the project have any code?
 	FGameProjectGenerationModule& GameProjectModule = FModuleManager::LoadModuleChecked<FGameProjectGenerationModule>(TEXT("GameProjectGeneration"));
-	bool bProjectHasCode = GameProjectModule.Get().ProjectHasCodeFiles();
+	bool bProjectHasCode = GameProjectModule.Get().ProjectRequiresBuild(InPlatformInfoName);
 
 	const PlatformInfo::FPlatformInfo* const PlatformInfo = PlatformInfo::FindPlatformInfo(InPlatformInfoName);
 	check(PlatformInfo);
diff --git a/Engine/Source/Editor/UnrealEd/Private/Kismet2/DebuggerCommands.cpp b/Engine/Source/Editor/UnrealEd/Private/Kismet2/DebuggerCommands.cpp
--- a/Engine/Source/Editor/UnrealEd/Private/Kismet2/DebuggerCommands.cpp
+++ b/Engine/Source/Editor/UnrealEd/Private/Kismet2/DebuggerCommands.cpp
@@ -1488,7 +1488,7 @@ bool FInternalPlayWorldCommandCallbacks::IsReadyToLaunchOnDevice(FString DeviceI
 	check(PlatformInfo);
 
 	FGameProjectGenerationModule& GameProjectModule = FModuleManager::LoadModuleChecked<FGameProjectGenerationModule>(TEXT("GameProjectGeneration"));
-	bool bHasCode = GameProjectModule.Get().ProjectHasCodeFiles();
+	bool bHasCode = GameProjectModule.Get().ProjectRequiresBuild(FName(*PlatformName));
 
 	if (PlatformInfo->SDKStatus == PlatformInfo::EPlatformSDKStatus::NotInstalled)
 	{
