Removed close button from the Convert Project modal dialog, as it was returning an undefined reason code (and was not necessary as there was already a Cancel button).

diff --git a/Engine/Source/Editor/GameProjectGeneration/Private/SVerbChoiceDialog.cpp b/Engine/Source/Editor/GameProjectGeneration/Private/SVerbChoiceDialog.cpp
--- a/Engine/Source/Editor/GameProjectGeneration/Private/SVerbChoiceDialog.cpp
+++ b/Engine/Source/Editor/GameProjectGeneration/Private/SVerbChoiceDialog.cpp
@@ -16,6 +16,7 @@ int32 SVerbChoiceDialog::ShowModal( const FText& InTitle, const FText& InMessage
 		.Title( InTitle )
 		.SizingRule( ESizingRule::Autosized )
 		.AutoCenter(EAutoCenter::PreferredWorkArea)
+		.HasCloseButton(false)
 		.SupportsMinimize(false) 
 		.SupportsMaximize(false);
 
