fix for flipped logic when checking settings are different from default
UE-15280

diff --git a/Engine/Source/Editor/GameProjectGeneration/Private/GameProjectUtils.cpp b/Engine/Source/Editor/GameProjectGeneration/Private/GameProjectUtils.cpp
--- a/Engine/Source/Editor/GameProjectGeneration/Private/GameProjectUtils.cpp
+++ b/Engine/Source/Editor/GameProjectGeneration/Private/GameProjectUtils.cpp
@@ -2989,7 +2989,7 @@ bool GameProjectUtils::DoProjectSettingsMatchDefault(const FString& InPlatformNa
 			FString Default(TEXT("False")), Project(TEXT("False"));
 			DefaultIni.GetString(*InSection, *((*InBoolKeys)[Index]), Default);
 			ProjIni.GetString(*InSection, *((*InBoolKeys)[Index]), Project);
-			if (!Default.Compare(Project, ESearchCase::IgnoreCase))
+			if (Default.Compare(Project, ESearchCase::IgnoreCase))
 			{
 				return false;
 			}
@@ -3017,7 +3017,7 @@ bool GameProjectUtils::DoProjectSettingsMatchDefault(const FString& InPlatformNa
 			FString Default(TEXT("False")), Project(TEXT("False"));
 			DefaultIni.GetString(*InSection, *((*InStringKeys)[Index]), Default);
 			ProjIni.GetString(*InSection, *((*InStringKeys)[Index]), Project);
-			if (!Default.Compare(Project, ESearchCase::IgnoreCase))
+			if (Default.Compare(Project, ESearchCase::IgnoreCase))
 			{
 				return false;
 			}
