don't allow Tappy Chicken to be considered a code based project in binary release
UE-13689
#ios

diff --git a/Engine/Source/Editor/GameProjectGeneration/Private/GameProjectUtils.cpp b/Engine/Source/Editor/GameProjectGeneration/Private/GameProjectUtils.cpp
--- a/Engine/Source/Editor/GameProjectGeneration/Private/GameProjectUtils.cpp
+++ b/Engine/Source/Editor/GameProjectGeneration/Private/GameProjectUtils.cpp
@@ -2966,12 +2966,15 @@ bool GameProjectUtils::ProjectRequiresBuild(const FName InPlatformInfoName)
 	}
 
 	bool bRequiresBuild = false;
-	// check to see if the default build settings have changed
-	bRequiresBuild |= !HasDefaultBuildSettings(InPlatformInfoName);
 
-	// check to see if any plugins beyond the defaults have been enabled
-	bRequiresBuild |= IProjectManager::Get().IsNonDefaultPluginEnabled();
+	if (!FRocketSupport::IsRocket())
+	{
+		// check to see if the default build settings have changed
+		bRequiresBuild |= !HasDefaultBuildSettings(InPlatformInfoName);
 
+		// check to see if any plugins beyond the defaults have been enabled
+		bRequiresBuild |= IProjectManager::Get().IsNonDefaultPluginEnabled();
+	}
 	return bRequiresBuild;
 }
 
