Improved the localization support of the class wizard dialogs

UE-14158 - Description of Parent Object Class

Fixed some issues where getting a localized string for a class tooltip would fail, and updated the class wizard to use FText so it can display localized descriptions.

diff --git a/Engine/Source/Editor/GameProjectGeneration/Private/GameProjectUtils.cpp b/Engine/Source/Editor/GameProjectGeneration/Private/GameProjectUtils.cpp
--- a/Engine/Source/Editor/GameProjectGeneration/Private/GameProjectUtils.cpp
+++ b/Engine/Source/Editor/GameProjectGeneration/Private/GameProjectUtils.cpp
@@ -44,29 +44,30 @@ static_assert(PLATFORM_MAX_FILEPATH_LENGTH - MAX_PROJECT_PATH_BUFFER_SPACE > 0,
 TWeakPtr<SNotificationItem> GameProjectUtils::UpdateGameProjectNotification = NULL;
 TWeakPtr<SNotificationItem> GameProjectUtils::WarningProjectNameNotification = NULL;
 
-FString FNewClassInfo::GetClassName() const
+FText FNewClassInfo::GetClassName() const
 {
 	switch(ClassType)
 	{
 	case EClassType::UObject:
-		return BaseClass ? FName::NameToDisplayString(BaseClass->GetName(), false) : TEXT("");
+		return BaseClass ? BaseClass->GetDisplayNameText() : FText::GetEmpty();
 
 	case EClassType::EmptyCpp:
-		return TEXT("None");
+		return LOCTEXT("NoParentClass", "None");
 
 	case EClassType::SlateWidget:
-		return TEXT("Slate Widget");
+		return LOCTEXT("SlateWidgetParentClass", "Slate Widget");
 
 	case EClassType::SlateWidgetStyle:
-		return TEXT("Slate Widget Style");
+		return LOCTEXT("SlateWidgetStyleParentClass", "Slate Widget Style");
 
 	default:
 		break;
 	}
-	return TEXT("");
+
+	return FText::GetEmpty();
 }
 
-FString FNewClassInfo::GetClassDescription() const
+FText FNewClassInfo::GetClassDescription(const bool bFullDescription/* = true*/) const
 {
 	switch(ClassType)
 	{
@@ -74,34 +75,40 @@ FString FNewClassInfo::GetClassDescription() const
 		{
 			if(BaseClass)
 			{
-				FString ClassDescription = BaseClass->GetToolTipText().ToString();
-				int32 FullStopIndex = 0;
-				if(ClassDescription.FindChar('.', FullStopIndex))
+				FString ClassDescription = BaseClass->GetToolTipText(/*bShortTooltip=*/!bFullDescription).ToString();
+
+				if(!bFullDescription)
 				{
-					// Only show the first sentence so as not to clutter up the UI with a detailed description of implementation details
-					ClassDescription = ClassDescription.Left(FullStopIndex + 1);
+					int32 FullStopIndex = 0;
+					if(ClassDescription.FindChar('.', FullStopIndex))
+					{
+						// Only show the first sentence so as not to clutter up the UI with a detailed description of implementation details
+						ClassDescription = ClassDescription.Left(FullStopIndex + 1);
+					}
+
+					// Strip out any new-lines in the description
+					ClassDescription.ReplaceInline(TEXT("\n"), TEXT(" "));	
 				}
 
-				// Strip out any new-lines in the description
-				ClassDescription = ClassDescription.Replace(TEXT("\n"), TEXT(" "));
-				return ClassDescription;
+				return FText::FromString(ClassDescription);
 			}
 		}
 		break;
 
 	case EClassType::EmptyCpp:
-		return TEXT("An empty C++ class with a default constructor and destructor");
+		return LOCTEXT("EmptyClassDescription", "An empty C++ class with a default constructor and destructor");
 
 	case EClassType::SlateWidget:
-		return TEXT("A custom Slate widget, deriving from SCompoundWidget");
+		return LOCTEXT("SlateWidgetClassDescription", "A custom Slate widget, deriving from SCompoundWidget");
 
 	case EClassType::SlateWidgetStyle:
-		return TEXT("A custom Slate widget style, deriving from FSlateWidgetStyle, along with its associated UObject wrapper class");
+		return LOCTEXT("SlateWidgetStyleClassDescription", "A custom Slate widget style, deriving from FSlateWidgetStyle, along with its associated UObject wrapper class");
 
 	default:
 		break;
 	}
-	return TEXT("");
+
+	return FText::GetEmpty();
 }
 
 const FSlateBrush* FNewClassInfo::GetClassIcon() const
diff --git a/Engine/Source/Editor/GameProjectGeneration/Private/SNewClassDialog.cpp b/Engine/Source/Editor/GameProjectGeneration/Private/SNewClassDialog.cpp
--- a/Engine/Source/Editor/GameProjectGeneration/Private/SNewClassDialog.cpp
+++ b/Engine/Source/Editor/GameProjectGeneration/Private/SNewClassDialog.cpp
@@ -814,8 +814,9 @@ TSharedRef<ITableRow> SNewClassDialog::MakeParentClassListViewWidget(TSharedPtr<
 		return SNew( STableRow<TSharedPtr<FParentClassItem>>, OwnerTable );
 	}
 
-	const FString ClassName = ParentClassItem->ParentClassInfo.GetClassName();
-	const FString ClassDescription = ParentClassItem->ParentClassInfo.GetClassDescription();
+	const FText ClassName = ParentClassItem->ParentClassInfo.GetClassName();
+	const FText ClassFullDescription = ParentClassItem->ParentClassInfo.GetClassDescription(/*bFullDescription*/true);
+	const FText ClassShortDescription = ParentClassItem->ParentClassInfo.GetClassDescription(/*bFullDescription*/false);
 	const FSlateBrush* const ClassBrush = ParentClassItem->ParentClassInfo.GetClassIcon();
 	const UClass* Class = ParentClassItem->ParentClassInfo.BaseClass;
 
@@ -824,7 +825,7 @@ TSharedRef<ITableRow> SNewClassDialog::MakeParentClassListViewWidget(TSharedPtr<
 	return
 		SNew( STableRow<TSharedPtr<FParentClassItem>>, OwnerTable )
 		.Style(FEditorStyle::Get(), "NewClassDialog.ParentClassListView.TableRow")
-		.ToolTip(IDocumentation::Get()->CreateToolTip(FText::FromString(ClassDescription), nullptr, FEditorClassUtils::GetDocumentationPage(Class), FEditorClassUtils::GetDocumentationExcerpt(Class)))
+		.ToolTip(IDocumentation::Get()->CreateToolTip(ClassFullDescription, nullptr, FEditorClassUtils::GetDocumentationPage(Class), FEditorClassUtils::GetDocumentationExcerpt(Class)))
 		[
 			SNew(SBox).HeightOverride(ItemHeight)
 			[
@@ -850,7 +851,7 @@ TSharedRef<ITableRow> SNewClassDialog::MakeParentClassListViewWidget(TSharedPtr<
 					[
 						SNew(STextBlock)
 						.TextStyle( FEditorStyle::Get(), "NewClassDialog.ParentClassItemTitle" )
-						.Text(FText::FromString(ClassName))
+						.Text(ClassName)
 					]
 				]
 
@@ -860,7 +861,7 @@ TSharedRef<ITableRow> SNewClassDialog::MakeParentClassListViewWidget(TSharedPtr<
 				[
 					SNew(STextBlock)
 					//.AutoWrapText(true)
-					.Text(FText::FromString(ClassDescription))
+					.Text(ClassShortDescription)
 				]
 			]
 		];
@@ -868,7 +869,7 @@ TSharedRef<ITableRow> SNewClassDialog::MakeParentClassListViewWidget(TSharedPtr<
 
 FText SNewClassDialog::GetSelectedParentClassName() const
 {
-	return ParentClassInfo.IsSet() ? FText::FromString(ParentClassInfo.GetClassName()) : FText::GetEmpty();
+	return ParentClassInfo.IsSet() ? ParentClassInfo.GetClassName() : FText::GetEmpty();
 }
 
 FString GetClassHeaderPath(const UClass* Class)
diff --git a/Engine/Source/Editor/GameProjectGeneration/Public/AddToProjectConfig.h b/Engine/Source/Editor/GameProjectGeneration/Public/AddToProjectConfig.h
--- a/Engine/Source/Editor/GameProjectGeneration/Public/AddToProjectConfig.h
+++ b/Engine/Source/Editor/GameProjectGeneration/Public/AddToProjectConfig.h
@@ -61,10 +61,10 @@ struct FNewClassInfo
 	}
 
 	/** Get the "friendly" class name to use in the UI */
-	FString GetClassName() const;
+	FText GetClassName() const;
 
 	/** Get the class description to use in the UI */
-	FString GetClassDescription() const;
+	FText GetClassDescription(const bool bFullDescription = true) const;
 
 	/** Get the class icon to use in the UI */
 	const FSlateBrush* GetClassIcon() const;
diff --git a/Engine/Source/Runtime/CoreUObject/Private/UObject/Class.cpp b/Engine/Source/Runtime/CoreUObject/Private/UObject/Class.cpp
--- a/Engine/Source/Runtime/CoreUObject/Private/UObject/Class.cpp
+++ b/Engine/Source/Runtime/CoreUObject/Private/UObject/Class.cpp
@@ -239,7 +239,7 @@ FText UField::GetToolTipText(bool bShortTooltip) const
 
 	const FString Namespace = bFoundShortTooltip ? TEXT("UObjectShortTooltips") : TEXT("UObjectToolTips");
 	const FString Key = GetFullGroupName(false);
-	if ( !(FText::FindText( Namespace, Key, /*OUT*/LocalizedToolTip )) || *FTextInspector::GetSourceString(LocalizedToolTip) != NativeToolTip)
+	if ( !FText::FindText( Namespace, Key, /*OUT*/LocalizedToolTip, &NativeToolTip ) )
 	{
 		if (NativeToolTip.IsEmpty())
 		{
diff --git a/Engine/Source/Runtime/CoreUObject/Private/UObject/Enum.cpp b/Engine/Source/Runtime/CoreUObject/Private/UObject/Enum.cpp
--- a/Engine/Source/Runtime/CoreUObject/Private/UObject/Enum.cpp
+++ b/Engine/Source/Runtime/CoreUObject/Private/UObject/Enum.cpp
@@ -358,7 +358,7 @@ FText UEnum::GetToolTipText(int32 NameIndex) const
 		?			GetFullGroupName(false)
 		:			GetFullGroupName(false) + TEXT(".") + GetEnumName(NameIndex);
 		
-	if ( !(FText::FindText( Namespace, Key, /*OUT*/LocalizedToolTip )) || *FTextInspector::GetSourceString(LocalizedToolTip) != NativeToolTip)
+	if ( !FText::FindText( Namespace, Key, /*OUT*/LocalizedToolTip, &NativeToolTip ) )
 	{
 		static const FString DoxygenSee(TEXT("@see"));
 		if (NativeToolTip.Split(DoxygenSee, &NativeToolTip, nullptr, ESearchCase::IgnoreCase, ESearchDir::FromStart))
