Fixed users unable to create blueprinted components without visual studio installed.  The error message was incorrectly appearing

https://jira.ol.epicgames.net/browse/UE-16145

diff --git a/Engine/Source/Editor/GameProjectGeneration/Private/SNewClassDialog.cpp b/Engine/Source/Editor/GameProjectGeneration/Private/SNewClassDialog.cpp
--- a/Engine/Source/Editor/GameProjectGeneration/Private/SNewClassDialog.cpp
+++ b/Engine/Source/Editor/GameProjectGeneration/Private/SNewClassDialog.cpp
@@ -997,7 +997,7 @@ EVisibility SNewClassDialog::GetGlobalErrorLabelVisibility() const
 
 FText SNewClassDialog::GetGlobalErrorLabelText() const
 {
-	if ( !FSourceCodeNavigation::IsCompilerAvailable() )
+	if ( ClassDomain == EClassDomain::Native && !FSourceCodeNavigation::IsCompilerAvailable() )
 	{
 		return FText::Format( LOCTEXT("NoCompilerFound", "No compiler was found. In order to use C++ code, you must first install {0}."), FSourceCodeNavigation::GetSuggestedSourceCodeIDE() );
 	}
@@ -1096,7 +1096,7 @@ void SNewClassDialog::CancelClicked()
 
 bool SNewClassDialog::CanFinish() const
 {
-	return bLastInputValidityCheckSuccessful && ParentClassInfo.IsSet() && FSourceCodeNavigation::IsCompilerAvailable();
+	return bLastInputValidityCheckSuccessful && ParentClassInfo.IsSet() && (ClassDomain == EClassDomain::Blueprint || FSourceCodeNavigation::IsCompilerAvailable());
 }
 
 void SNewClassDialog::FinishClicked()
