Fixed another Feature Pack detail enum

diff --git a/Engine/Source/Editor/GameProjectGeneration/Private/GameProjectUtils.cpp b/Engine/Source/Editor/GameProjectGeneration/Private/GameProjectUtils.cpp
--- a/Engine/Source/Editor/GameProjectGeneration/Private/GameProjectUtils.cpp
+++ b/Engine/Source/Editor/GameProjectGeneration/Private/GameProjectUtils.cpp
@@ -3515,7 +3515,7 @@ bool GameProjectUtils::AddSharedContentToProject(const FProjectInformation &InPr
 				{
 					// Build the Packname from the mount and detail
 					FString DetailString;
-					UEnum::GetValueAsString(TEXT("/Script/GameProjectGeneration.ETemplateDetailLevel"), EachRequiredDetail, DetailString);
+					UEnum::GetValueAsString(TEXT("/Script/AddContentDialog.EFeaturePackDetailLevel"), EachRequiredDetail, DetailString);
 					
 					FString FullPackFilename = FPaths::FeaturePackDir() + EachPack.MountName +  DetailString + DefaultFeaturePackExtension;
 					if (FPaths::FileExists(FullPackFilename) == false)
diff --git a/Engine/Source/Editor/UnrealEd/Private/UnrealEdMisc.cpp b/Engine/Source/Editor/UnrealEd/Private/UnrealEdMisc.cpp
--- a/Engine/Source/Editor/UnrealEd/Private/UnrealEdMisc.cpp
+++ b/Engine/Source/Editor/UnrealEd/Private/UnrealEdMisc.cpp
@@ -1650,7 +1650,7 @@ void FUnrealEdMisc::MountTemplateSharedPaths()
 			}
 			else if (EachPack.DetailLevels.Num() == 0)
 			{
-				// If no levels are supplied, assume standard
+				// If no levels are supplied we cant really use this pack !
 				UE_LOG(LogUnrealEdMisc, Warning, TEXT("No detail levels defined for %s in %s."), *EachPack.MountName, *TemplateFilename );
 				continue;
 			}
