Added automation test for project creation.
Added support for which shared resource to edit

diff --git a/Engine/Source/Editor/AddContentDialog/Private/IContentSource.h b/Engine/Source/Editor/AddContentDialog/Private/IContentSource.h
--- a/Engine/Source/Editor/AddContentDialog/Private/IContentSource.h
+++ b/Engine/Source/Editor/AddContentDialog/Private/IContentSource.h
@@ -11,6 +11,7 @@ enum class EContentSourceCategory:uint8
 	BlueprintFeature,
 	CodeFeature,
 	Content,
+	SharedPack,
 	Unknown
 };
 
diff --git a/Engine/Source/Editor/GameProjectGeneration/Classes/TemplateProjectDefs.h b/Engine/Source/Editor/GameProjectGeneration/Classes/TemplateProjectDefs.h
--- a/Engine/Source/Editor/GameProjectGeneration/Classes/TemplateProjectDefs.h
+++ b/Engine/Source/Editor/GameProjectGeneration/Classes/TemplateProjectDefs.h
@@ -97,13 +97,18 @@ class UTemplateProjectDefs : public UObject
 	UPROPERTY(config)
 	FString AssetTypes;
 
+	/* Should we allow creation of a project from this template */
 	UPROPERTY(config)
-	bool bHasFeaturePack;
+	bool bAllowProjectCreation;
 	
 	/* Optional list of feature packs to include */
 	UPROPERTY(config)
 	TArray<FString> PacksToInclude;
 
+	/** What detail level to edit when editing shared template resources  */
+	UPROPERTY(config)
+	EFeaturePackDetailLevel EditDetailLevelPreference;
+
 	/* Shared feature packs. The files in these packs listed in these structures marked as 'additionalfiles' will be copied on project generation*/
 	UPROPERTY(config)
 	TArray<FFeaturePackLevelSet> SharedContentPacks;
diff --git a/Engine/Source/Editor/GameProjectGeneration/Private/GameProjectUtils.cpp b/Engine/Source/Editor/GameProjectGeneration/Private/GameProjectUtils.cpp
--- a/Engine/Source/Editor/GameProjectGeneration/Private/GameProjectUtils.cpp
+++ b/Engine/Source/Editor/GameProjectGeneration/Private/GameProjectUtils.cpp
@@ -610,7 +610,7 @@ void GameProjectUtils::GetStarterContentFiles(TArray<FString>& OutFilenames)
 	IFileManager::Get().FindFilesRecursive(OutFilenames, *SrcFolder, *SearchPath, /*Files=*/true, /*Directories=*/false);
 }
 
-bool GameProjectUtils::CreateProject(const FProjectInformation& InProjectInfo, FText& OutFailReason, FText& OutFailLog)
+bool GameProjectUtils::CreateProject(const FProjectInformation& InProjectInfo, FText& OutFailReason, FText& OutFailLog, TArray<FString>* OutCreatedFiles)
 {
 	if ( !IsValidProjectFileForCreation(InProjectInfo.ProjectFilename, OutFailReason) )
 	{
@@ -629,7 +629,7 @@ bool GameProjectUtils::CreateProject(const FProjectInformation& InProjectInfo, F
 	}
 	else
 	{
-		bProjectCreationSuccessful = CreateProjectFromTemplate(InProjectInfo, OutFailReason, OutFailLog);
+		bProjectCreationSuccessful = CreateProjectFromTemplate(InProjectInfo, OutFailReason, OutFailLog, OutCreatedFiles);
 		TemplateName = FPaths::GetBaseFilename(InProjectInfo.TemplateFile);
 	}
 
@@ -1161,7 +1161,7 @@ bool GameProjectUtils::GenerateProjectFromScratch(const FProjectInformation& InP
 	return true;
 }
 
-bool GameProjectUtils::CreateProjectFromTemplate(const FProjectInformation& InProjectInfo, FText& OutFailReason, FText& OutFailLog)
+bool GameProjectUtils::CreateProjectFromTemplate(const FProjectInformation& InProjectInfo, FText& OutFailReason, FText& OutFailLog,TArray<FString>* OutCreatedFiles)
 {
 	FScopedSlowTask SlowTask(10);
 
@@ -1592,6 +1592,10 @@ bool GameProjectUtils::CreateProjectFromTemplate(const FProjectInformation& InPr
 		return false;
 	}
 	
+	if( OutCreatedFiles != nullptr )
+	{
+		OutCreatedFiles->Append(CreatedFiles);
+	}
 	return true;
 }
 
diff --git a/Engine/Source/Editor/GameProjectGeneration/Private/SNewProjectWizard.cpp b/Engine/Source/Editor/GameProjectGeneration/Private/SNewProjectWizard.cpp
--- a/Engine/Source/Editor/GameProjectGeneration/Private/SNewProjectWizard.cpp
+++ b/Engine/Source/Editor/GameProjectGeneration/Private/SNewProjectWizard.cpp
@@ -51,6 +51,20 @@ struct FTemplateItem
 		: Name(InName), Description(InDescription), bGenerateCode(bInGenerateCode), Type(InType), SortKey(MoveTemp(InSortKey)), ProjectFile(MoveTemp(InProjectFile)), Thumbnail(InThumbnail), PreviewImage(InPreviewImage)
 		, ClassTypes(InClassTypes), AssetTypes(InAssetTypes)
 	{}
+
+	FTemplateItem(const FTemplateItem& InItem)
+	{
+		Name = InItem.Name;
+		Description = InItem.Description;
+		bGenerateCode = InItem.bGenerateCode;
+		Type = InItem.Type;
+
+		SortKey = InItem.SortKey;
+		ProjectFile = InItem.ProjectFile;
+
+		ClassTypes = InItem.ClassTypes;
+		AssetTypes = InItem.AssetTypes;
+	}
 };
 
 /**
@@ -1131,10 +1145,13 @@ FText SNewProjectWizard::GetNameAndLocationErrorLabelText() const
 	return FText::GetEmpty();
 }
 
-void SNewProjectWizard::FindTemplateProjects()
+TMap<FName, TArray<TSharedPtr<FTemplateItem>> >& SNewProjectWizard::FindTemplateProjects()
 {
 	// Default to showing the blueprint category
 	ActiveCategory = FTemplateCategory::BlueprintCategoryName;
+	
+	// Clear the list out first - or we could end up with duplicates
+	Templates.Empty();
 
 	// Add some default non-data driven templates
 	Templates.FindOrAdd(FTemplateCategory::BlueprintCategoryName).Add(MakeShareable(new FTemplateItem(
@@ -1211,6 +1228,9 @@ void SNewProjectWizard::FindTemplateProjects()
 				UTemplateProjectDefs* TemplateDefs = GameProjectUtils::LoadTemplateDefs(Root);
 				if ( TemplateDefs )
 				{
+					// Ignore any templates whoose definition says we cannot use to create a project
+					if( TemplateDefs->bAllowProjectCreation == false )
+						continue;
 					// Found a template. Add it to the template items list.
 					const FString ProjectFilename = Root / FoundProjectFiles[0];
 					FText TemplateName = TemplateDefs->GetDisplayNameText();
@@ -1282,6 +1302,7 @@ void SNewProjectWizard::FindTemplateProjects()
 			}
 		}
 	}
+	return Templates;
 }
 
 
diff --git a/Engine/Source/Editor/GameProjectGeneration/Private/SNewProjectWizard.h b/Engine/Source/Editor/GameProjectGeneration/Private/SNewProjectWizard.h
--- a/Engine/Source/Editor/GameProjectGeneration/Private/SNewProjectWizard.h
+++ b/Engine/Source/Editor/GameProjectGeneration/Private/SNewProjectWizard.h
@@ -23,6 +23,9 @@ public:
 
 	virtual void Tick( const FGeometry& AllottedGeometry, const double InCurrentTime, const float InDeltaTime ) override;
 
+	/** Populates TemplateItemsSource with templates found on disk */
+	TMap<FName, TArray<TSharedPtr<FTemplateItem>> >& FindTemplateProjects();
+
 private:
 
 	/** Build the set of template category tabs */
@@ -135,6 +138,8 @@ private:
 
 	/** Returns the tooltip text (actual warning) of the starter content warning */
 	FText GetStarterContentWarningTooltip() const;
+
+
 private:
 
 	EHardwareClass::Type SelectedHardwareClassTarget;
@@ -145,11 +150,6 @@ private:
 	void SetGraphicsPreset(EGraphicsPreset::Type InGraphicsPreset);
 	EGraphicsPreset::Type GetGraphicsPreset() const { return SelectedGraphicsPreset; }
 
-private:
-
-	/** Populates TemplateItemsSource with templates found on disk */
-	void FindTemplateProjects();
-
 	/** Sets the default project name and path */
 	void SetDefaultProjectLocation();
 
diff --git a/Engine/Source/Editor/GameProjectGeneration/Private/TemplateProjectDefs.cpp b/Engine/Source/Editor/GameProjectGeneration/Private/TemplateProjectDefs.cpp
--- a/Engine/Source/Editor/GameProjectGeneration/Private/TemplateProjectDefs.cpp
+++ b/Engine/Source/Editor/GameProjectGeneration/Private/TemplateProjectDefs.cpp
@@ -57,7 +57,8 @@ FTemplateConfigValue::FTemplateConfigValue(const FString& InFile, const FString&
 UTemplateProjectDefs::UTemplateProjectDefs(const FObjectInitializer& ObjectInitializer)
 	: Super(ObjectInitializer)
 {
-
+	bAllowProjectCreation = true;
+	EditDetailLevelPreference = EFeaturePackDetailLevel::Standard;
 }
 
 void UTemplateProjectDefs::FixupStrings(const FString& TemplateName, const FString& ProjectName)
diff --git a/Engine/Source/Editor/GameProjectGeneration/Private/Tests/GameProjectAutomationTests.cpp b/Engine/Source/Editor/GameProjectGeneration/Private/Tests/GameProjectAutomationTests.cpp
--- a/Engine/Source/Editor/GameProjectGeneration/Private/Tests/GameProjectAutomationTests.cpp
+++ b/Engine/Source/Editor/GameProjectGeneration/Private/Tests/GameProjectAutomationTests.cpp
@@ -50,6 +50,124 @@ namespace GameProjectAutomationUtils
 
 		return ProjectFilename;
 	}
+
+
+	/* 
+	 * Create a project from a templalte with a given criteria
+	 * @oaram	InTemplates			List of available project templates
+	 * @param	InTargetedHardware	Target hardware (EHardwareClass)
+	 * @param	InGraphicPreset		Graphics preset (EGraphicsPreset)
+	 * @param	InCategory			Target category (EContentSourceCategory)
+	 * @param	bInCopyStarterContent Should the starter content be copied also
+	 * @param	OutMatchedProjects	Total projects matching criteria
+	 * @param	OutCreatedProjects	Total projects succesfully created
+	 */
+	static void CreateProjectSet(TMap<FName, TArray<TSharedPtr<FTemplateItem>> >& InTemplates, EHardwareClass::Type InTargetedHardware, EGraphicsPreset::Type InGraphicPreset, EContentSourceCategory InCategory, bool bInCopyStarterContent, int32 OutMatchedProjects, int32 OutCreatedProjects)
+	{		
+		// If this is empty, it will use the same name for each project, otherwise it will create a project based on target platform and source template
+		FString TestRootFolder;// = "ProjectTests";
+		
+		// This has the code remove the projects once created
+		bool bRemoveCreatedProjects = true;		
+		OutCreatedProjects = 0;
+		OutMatchedProjects = 0;
+		UEnum* SourceCategoryEnum = FindObject<UEnum>(ANY_PACKAGE, TEXT("EContentSourceCategory"));
+
+		// The category name int he FTemplateItem is not the same as the enum definition EContentSourceCategory - convert it
+		FName CategoryName;
+		if( InCategory == EContentSourceCategory::BlueprintFeature)
+		{
+			CategoryName = FTemplateCategory::BlueprintCategoryName;
+		}
+		else if (InCategory == EContentSourceCategory::CodeFeature)
+		{
+			CategoryName = FTemplateCategory::CodeCategoryName;
+		}
+		else
+		{
+			// We didnt match a category
+			if (SourceCategoryEnum == nullptr)
+			{
+				UE_LOG(LogGameProjectGenerationTests, Display, TEXT("Test failed! Unknown category type %d"), *SourceCategoryEnum->GetEnumName(int32(InCategory)));
+			}
+			else
+			{
+				UE_LOG(LogGameProjectGenerationTests, Display, TEXT("Test failed! Unknown category type %d"), (int32)(InCategory));
+			}
+			
+			return;
+		}
+
+		// Iterate all templates and try to create those that match the required criteria
+		for (auto EachTemplate : InTemplates)
+		{
+			FName Name = EachTemplate.Key;						
+			// Check this is the correct category
+			if (Name == CategoryName)
+			{
+				// Now iterate each template in the category
+				for (TSharedPtr<FTemplateItem> OneTemplate : EachTemplate.Value)
+				{				
+					TSharedPtr<FTemplateItem> Item = OneTemplate;
+					// If this template is flagged as not for creation dont try to create it
+					if( Item->ProjectFile.IsEmpty())
+						continue;
+
+					FString DesiredProjectFilename;
+					if( TestRootFolder.IsEmpty() == true)
+					{
+						// Same name for all
+						DesiredProjectFilename = GameProjectAutomationUtils::GetDesiredProjectFilename();
+					}
+					else
+					{
+						// Unique names
+						FString Hardware;
+						if (InTargetedHardware == EHardwareClass::Desktop)
+						{
+							Hardware = TEXT("Dsk");
+						}
+						else
+						{
+							Hardware = TEXT("Mob");
+						}
+						FString ProjectName = FPaths::GetCleanFilename(Item->ProjectFile).Replace(TEXT("TP_"), TEXT(""));
+						FString ProjectDirName = FPaths::GetBaseFilename(Item->ProjectFile).Replace(TEXT("TP_"), TEXT(""));
+						FString BasePath = FPaths::RootDir();
+						DesiredProjectFilename = FString::Printf(TEXT("%s/%s/%s%s/%s%s"), *BasePath, *TestRootFolder, *Hardware, *ProjectDirName, *Hardware, *ProjectName);
+					}
+					
+					// Setup creation parameters
+					FText FailReason, FailLog;
+					FProjectInformation ProjectInfo(DesiredProjectFilename, Item->bGenerateCode, bInCopyStarterContent, Item->ProjectFile);
+					ProjectInfo.TargetedHardware = InTargetedHardware;
+					ProjectInfo.DefaultGraphicsPerformance = InGraphicPreset;
+					TArray<FString> CreatedFiles;
+					OutMatchedProjects++;
+
+					// Finally try to create the project					
+					if (!GameProjectUtils::CreateProject(ProjectInfo, FailReason, FailLog, &CreatedFiles))
+					{
+						// Failed, report the reason
+						UE_LOG(LogGameProjectGenerationTests, Display, TEXT("Test failed! Failed to create %s project %s based on %s. Reason:%s"), *SourceCategoryEnum->GetEnumName((int32)InCategory), *DesiredProjectFilename, *Item->Name.ToString(), *FailReason.ToString());
+					}
+					else
+					{
+						// Created ok
+						OutCreatedProjects++;
+						
+						// Now remove the files we just created (if required)
+						if(bRemoveCreatedProjects == true)						
+						{
+							FString RootFolder = FPaths::GetPath(DesiredProjectFilename);
+							GameProjectUtils::DeleteCreatedFiles(RootFolder, CreatedFiles);
+						}
+					}
+				}
+			}
+		}
+		return;
+	}
 }
 
 /**
@@ -199,4 +317,61 @@ bool FBuildPromotionNewProjectMapTest::RunTest(const FString& Parameters)
 	}
 
 	return true;
+}
+
+/*
+* Template project creation test
+*/
+IMPLEMENT_SIMPLE_AUTOMATION_TEST( FCreateBPTemplateProjectAutomationTests, "System.Promotion.Project Promotion Pass.Step 3 NewProjectCreationTests.CreateBlueprintProjects", EAutomationTestFlags::ATF_Editor )
+
+/** 
+ * Uses the new project wizard to locate all templates available for new blueprint project creation and verifies creation succeeds.
+ *
+ * @param Parameters - Unused for this test
+ * @return	TRUE if the test was successful, FALSE otherwise
+ */
+bool FCreateBPTemplateProjectAutomationTests::RunTest(const FString& Parameters)
+{
+	TSharedPtr<SNewProjectWizard> NewProjectWizard;
+	NewProjectWizard = SNew(SNewProjectWizard);
+	//return ;
+	TMap<FName, TArray<TSharedPtr<FTemplateItem>> >& Templates = NewProjectWizard->FindTemplateProjects();//GameProjectAutomationUtils::CreateTemplateList();
+	int32 OutMatchedProjectsDesk = 0;
+	int32 OutCreatedProjectsDesk = 0;
+	GameProjectAutomationUtils::CreateProjectSet(Templates, EHardwareClass::Desktop, EGraphicsPreset::Maximum, EContentSourceCategory::BlueprintFeature, false, OutMatchedProjectsDesk, OutCreatedProjectsDesk);
+
+	int32 OutMatchedProjectsMob = 0;
+	int32 OutCreatedProjectsMob = 0;
+	GameProjectAutomationUtils::CreateProjectSet(Templates, EHardwareClass::Mobile, EGraphicsPreset::Maximum, EContentSourceCategory::BlueprintFeature, false, OutMatchedProjectsMob, OutCreatedProjectsMob);
+	
+	return ( OutMatchedProjectsDesk == OutCreatedProjectsDesk ) && ( OutMatchedProjectsMob == OutCreatedProjectsMob  );
+}
+
+/*
+* Template project creation test
+*/
+IMPLEMENT_SIMPLE_AUTOMATION_TEST( FCreateCPPTemplateProjectAutomationTests, "System.Promotion.Project Promotion Pass.Step 3 NewProjectCreationTests.CreateCodeProjects", EAutomationTestFlags::ATF_Editor )
+
+/** 
+ * Uses the new project wizard to locate all templates available for new code project creation and verifies creation succeeds.
+ *
+ * @param Parameters - Unused for this test
+ * @return	TRUE if the test was successful, FALSE otherwise
+ */
+bool FCreateCPPTemplateProjectAutomationTests::RunTest(const FString& Parameters)
+{
+	TSharedPtr<SNewProjectWizard> NewProjectWizard;
+	NewProjectWizard = SNew(SNewProjectWizard);
+	//return ;
+	TMap<FName, TArray<TSharedPtr<FTemplateItem>> >& Templates = NewProjectWizard->FindTemplateProjects();//GameProjectAutomationUtils::CreateTemplateList();
+	int32 OutMatchedProjectsDesk = 0;
+	int32 OutCreatedProjectsDesk = 0;
+	GameProjectAutomationUtils::CreateProjectSet(Templates, EHardwareClass::Desktop, EGraphicsPreset::Maximum, EContentSourceCategory::CodeFeature, false, OutMatchedProjectsDesk, OutCreatedProjectsDesk);
+
+	int32 OutMatchedProjectsMob = 0;
+	int32 OutCreatedProjectsMob = 0;
+	GameProjectAutomationUtils::CreateProjectSet(Templates, EHardwareClass::Mobile, EGraphicsPreset::Maximum, EContentSourceCategory::CodeFeature, false, OutMatchedProjectsMob, OutCreatedProjectsMob);
+
+	return (OutMatchedProjectsDesk == OutCreatedProjectsDesk) && (OutMatchedProjectsMob == OutCreatedProjectsMob);
+
 }
\ No newline at end of file
diff --git a/Engine/Source/Editor/GameProjectGeneration/Public/GameProjectUtils.h b/Engine/Source/Editor/GameProjectGeneration/Public/GameProjectUtils.h
--- a/Engine/Source/Editor/GameProjectGeneration/Public/GameProjectUtils.h
+++ b/Engine/Source/Editor/GameProjectGeneration/Public/GameProjectUtils.h
@@ -22,6 +22,7 @@ struct FProjectInformation
 		, DefaultGraphicsPerformance(EGraphicsPreset::Maximum)
 	{
 	}
+
 	FString ProjectFilename;
 	FString TemplateFile;
 
@@ -92,7 +93,7 @@ public:
 	static bool OpenCodeIDE(const FString& ProjectFile, FText& OutFailReason);
 
 	/** Creates the specified project file and all required folders. If TemplateFile is non-empty, it will be used as the template for creation. On failure, OutFailReason will be populated. */
-	static bool CreateProject(const FProjectInformation& InProjectInfo, FText& OutFailReason, FText& OutFailLog);
+	static bool CreateProject(const FProjectInformation& InProjectInfo, FText& OutFailReason, FText& OutFailLog, TArray<FString>* OutCreatedFiles = nullptr);
 
 	/** Prompts the user to update his project file, if necessary. */
 	static void CheckForOutOfDateGameProjectFile();
@@ -235,7 +236,10 @@ public:
 	static FString MakeIncludeList(const TArray<FString>& InList);
 
 	/** Returns true if the currently loaded project requires a code build */
-	static bool ProjectRequiresBuild(const FName InPlatformInfoName);
+	static bool ProjectRequiresBuild(const FName InPlatformInfoName);		
+
+	/** Deletes the specified list of files that were created during file creation */
+	static void DeleteCreatedFiles(const FString& RootFolder, const TArray<FString>& CreatedFiles);
 
 private:
 
@@ -245,7 +249,7 @@ private:
 	static bool GenerateProjectFromScratch(const FProjectInformation& InProjectInfo, FText& OutFailReason, FText& OutFailLog);
 
 	/** Generates a new project using a template project */
-	static bool CreateProjectFromTemplate(const FProjectInformation& InProjectInfo, FText& OutFailReason, FText& OutFailLog);
+	static bool CreateProjectFromTemplate(const FProjectInformation& InProjectInfo, FText& OutFailReason, FText& OutFailLog, TArray<FString>* OutCreatedFiles = nullptr);
 
 	/** Sets the engine association for a new project. Handles foreign and non-foreign projects. */
 	static bool SetEngineAssociationForForeignProject(const FString& ProjectFileName, FText& OutFailReason);
@@ -281,9 +285,6 @@ private:
 	/** Returns true if file cleanup on failure is enabled, false if not */
 	static bool CleanupIsEnabled();
 
-	/** Deletes the specified list of files that were created during file creation */
-	static void DeleteCreatedFiles(const FString& RootFolder, const TArray<FString>& CreatedFiles);
-
 	/** Deletes any files that were generated by the generate project files step */
 	static void DeleteGeneratedProjectFiles(const FString& NewProjectFile);
 
diff --git a/Engine/Source/Editor/UnrealEd/Private/UnrealEdMisc.cpp b/Engine/Source/Editor/UnrealEd/Private/UnrealEdMisc.cpp
--- a/Engine/Source/Editor/UnrealEd/Private/UnrealEdMisc.cpp
+++ b/Engine/Source/Editor/UnrealEd/Private/UnrealEdMisc.cpp
@@ -1632,15 +1632,14 @@ void FUnrealEdMisc::MountTemplateSharedPaths()
 	UTemplateProjectDefs* TemplateInfo = GameProjectUtils::LoadTemplateDefs(TemplateFilename);
 	
 	if( TemplateInfo != nullptr )
-	{
-		// Which detail level do we want to mount for editing
-		EFeaturePackDetailLevel EditDetail = EFeaturePackDetailLevel::Standard;
-		
+	{	
+		EFeaturePackDetailLevel EditDetail = TemplateInfo->EditDetailLevelPreference;
+
 		// Extract the mount names and insert mount points for each of the shared packs
 		TArray<FString> AddedMountSources;
 		for (int32 iShared = 0; iShared < TemplateInfo->SharedContentPacks.Num() ; iShared++)
 		{
-			EFeaturePackDetailLevel EachEditDetail = EditDetail;
+			EFeaturePackDetailLevel EachEditDetail = (EFeaturePackDetailLevel)EditDetail;
 			FFeaturePackLevelSet EachPack = TemplateInfo->SharedContentPacks[iShared];
 			if (EachPack.DetailLevels.Num() == 1)
 			{
