Fix non-unity

diff --git a/Engine/Source/Editor/GameProjectGeneration/Private/SNewProjectWizard.cpp b/Engine/Source/Editor/GameProjectGeneration/Private/SNewProjectWizard.cpp
--- a/Engine/Source/Editor/GameProjectGeneration/Private/SNewProjectWizard.cpp
+++ b/Engine/Source/Editor/GameProjectGeneration/Private/SNewProjectWizard.cpp
@@ -15,7 +15,7 @@
 #include "BreakIterator.h"
 #include "SHyperlink.h"
 #include "SOutputLogDialog.h"
-
+#include "TemplateItem.h"
 #include "Settings/EditorSettings.h"
 
 #define LOCTEXT_NAMESPACE "NewProjectWizard"
@@ -31,42 +31,6 @@ namespace NewProjectWizardDefs
 	const float ItemHeight = ItemWidth + 30;
 }
 
-/** Struct describing a single template project */
-struct FTemplateItem
-{
-	FText		Name;
-	FText		Description;
-	bool		bGenerateCode;
-	FName		Type;
-
-	FString		SortKey;
-	FString		ProjectFile;
-
-	TSharedPtr<FSlateBrush> Thumbnail;
-	TSharedPtr<FSlateBrush> PreviewImage;
-
-	FString		ClassTypes;
-	FString		AssetTypes;
-	FTemplateItem(FText InName, FText InDescription, bool bInGenerateCode, FName InType, FString InSortKey, FString InProjectFile, TSharedPtr<FSlateBrush> InThumbnail, TSharedPtr<FSlateBrush> InPreviewImage,FString InClassTypes, FString InAssetTypes)
-		: Name(InName), Description(InDescription), bGenerateCode(bInGenerateCode), Type(InType), SortKey(MoveTemp(InSortKey)), ProjectFile(MoveTemp(InProjectFile)), Thumbnail(InThumbnail), PreviewImage(InPreviewImage)
-		, ClassTypes(InClassTypes), AssetTypes(InAssetTypes)
-	{}
-
-	FTemplateItem(const FTemplateItem& InItem)
-	{
-		Name = InItem.Name;
-		Description = InItem.Description;
-		bGenerateCode = InItem.bGenerateCode;
-		Type = InItem.Type;
-
-		SortKey = InItem.SortKey;
-		ProjectFile = InItem.ProjectFile;
-
-		ClassTypes = InItem.ClassTypes;
-		AssetTypes = InItem.AssetTypes;
-	}
-};
-
 /**
  * Simple widget used to display a folder path, and a name of a file:
  * __________________________  ____________________
diff --git a/Engine/Source/Editor/GameProjectGeneration/Private/TemplateItem.h b/Engine/Source/Editor/GameProjectGeneration/Private/TemplateItem.h
new file mode 100644
--- /dev/null
+++ b/Engine/Source/Editor/GameProjectGeneration/Private/TemplateItem.h
@@ -0,0 +1,40 @@
+// Copyright 1998-2015 Epic Games, Inc. All Rights Reserved.
+
+#pragma once
+
+/** Struct describing a single template project */
+struct FTemplateItem
+{
+	FText		Name;
+	FText		Description;
+	bool		bGenerateCode;
+	FName		Type;
+
+	FString		SortKey;
+	FString		ProjectFile;
+
+	TSharedPtr<FSlateBrush> Thumbnail;
+	TSharedPtr<FSlateBrush> PreviewImage;
+
+	FString		ClassTypes;
+	FString		AssetTypes;
+	FTemplateItem(FText InName, FText InDescription, bool bInGenerateCode, FName InType, FString InSortKey, FString InProjectFile, TSharedPtr<FSlateBrush> InThumbnail, TSharedPtr<FSlateBrush> InPreviewImage,FString InClassTypes, FString InAssetTypes)
+		: Name(InName), Description(InDescription), bGenerateCode(bInGenerateCode), Type(InType), SortKey(MoveTemp(InSortKey)), ProjectFile(MoveTemp(InProjectFile)), Thumbnail(InThumbnail), PreviewImage(InPreviewImage)
+		, ClassTypes(InClassTypes), AssetTypes(InAssetTypes)
+	{}
+
+	FTemplateItem(const FTemplateItem& InItem)
+	{
+		Name = InItem.Name;
+		Description = InItem.Description;
+		bGenerateCode = InItem.bGenerateCode;
+		Type = InItem.Type;
+
+		SortKey = InItem.SortKey;
+		ProjectFile = InItem.ProjectFile;
+
+		ClassTypes = InItem.ClassTypes;
+		AssetTypes = InItem.AssetTypes;
+	}
+};
+
diff --git a/Engine/Source/Editor/GameProjectGeneration/Private/Tests/GameProjectAutomationTests.cpp b/Engine/Source/Editor/GameProjectGeneration/Private/Tests/GameProjectAutomationTests.cpp
--- a/Engine/Source/Editor/GameProjectGeneration/Private/Tests/GameProjectAutomationTests.cpp
+++ b/Engine/Source/Editor/GameProjectGeneration/Private/Tests/GameProjectAutomationTests.cpp
@@ -8,6 +8,8 @@
 #include "Tests/AutomationTestSettings.h"
 #include "AutomationEditorCommon.h"
 #include "GameFramework/PlayerStart.h"
+#include "../TemplateCategory.h"
+#include "../TemplateItem.h"
 
 DEFINE_LOG_CATEGORY_STATIC(LogGameProjectGenerationTests, Log, All);
 
