Improved the folder blacklist in GameProjectUtils::DuplicateProjectForUpgrade

UE-18319 - Collections reset when the engine upgrades project

GameProjectUtils::DuplicateProjectForUpgrade used to wholesale skip the Saved directory, which meant that copying your project could appear to lose your local collections, as well as any saved user-settings.

This change allows it to copy the parts of the Saved directory that we know aren't transient. The RelativeDirectoriesToSkip array may need updating as new Saved content is added in order to keep the copy as fast as possible.

diff --git a/Engine/Source/Editor/GameProjectGeneration/Private/GameProjectUtils.cpp b/Engine/Source/Editor/GameProjectGeneration/Private/GameProjectUtils.cpp
--- a/Engine/Source/Editor/GameProjectGeneration/Private/GameProjectUtils.cpp
+++ b/Engine/Source/Editor/GameProjectGeneration/Private/GameProjectUtils.cpp
@@ -2274,35 +2274,73 @@ GameProjectUtils::EProjectDuplicateResult GameProjectUtils::DuplicateProjectForU
 		NewDirectoryName = FString::Printf(TEXT("%s - %d"), *BaseDirectoryName, Idx);
 	}
 
-	// Find all the root directory names
-	TArray<FString> RootDirectoryNames;
-	IFileManager::Get().FindFiles(RootDirectoryNames, *(OldDirectoryName / TEXT("*")), false, true);
+	// Recursively find all the files we need to copy, excluding those that are within the directories listed in SourceDirectoriesToSkip
+	struct FGatherFilesToCopyHelper
+	{
+	public:
+		FGatherFilesToCopyHelper(FString InRootSourceDirectory)
+			: RootSourceDirectory(MoveTemp(InRootSourceDirectory))
+		{
+			static const FString RelativeDirectoriesToSkip[] = {
+				TEXT("Binaries"),
+				TEXT("DerivedDataCache"),
+				TEXT("Intermediate"),
+				TEXT("Saved/Autosaves"),
+				TEXT("Saved/Backup"),
+				TEXT("Saved/Cooked"),
+				TEXT("Saved/HardwareSurvey"),
+				TEXT("Saved/Logs"),
+				TEXT("Saved/StagedBuilds"),
+			};
+
+			SourceDirectoriesToSkip.Reserve(ARRAY_COUNT(RelativeDirectoriesToSkip));
+			for (const FString& RelativeDirectoryToSkip : RelativeDirectoriesToSkip)
+			{
+				SourceDirectoriesToSkip.Emplace(RootSourceDirectory / RelativeDirectoryToSkip);
+			}
+		}
 
-	// Find all the source directories
-	TArray<FString> SourceDirectories;
-	SourceDirectories.Add(OldDirectoryName);
-	for(int32 Idx = 0; Idx < RootDirectoryNames.Num(); Idx++)
-	{
-		if(RootDirectoryNames[Idx] != TEXT("Binaries") && RootDirectoryNames[Idx] != TEXT("Intermediate") && RootDirectoryNames[Idx] != TEXT("Saved"))
+		void GatherFilesToCopy(TArray<FString>& OutSourceDirectories, TArray<FString>& OutSourceFiles)
 		{
-			FString SourceDirectory = OldDirectoryName / RootDirectoryNames[Idx];
-			SourceDirectories.Add(SourceDirectory);
-			IFileManager::Get().FindFilesRecursive(SourceDirectories, *SourceDirectory, TEXT("*"), false, true, false);
+			GatherFilesToCopy(RootSourceDirectory, OutSourceDirectories, OutSourceFiles);
 		}
-	}
 
-	// Find all the source files
-	TArray<FString> SourceFiles;
-	for(int32 Idx = 0; Idx < SourceDirectories.Num(); Idx++)
-	{
-		TArray<FString> SourceNames;
-		IFileManager::Get().FindFiles(SourceNames, *(SourceDirectories[Idx] / TEXT("*")), true, false);
-
-		for(int32 NameIdx = 0; NameIdx < SourceNames.Num(); NameIdx++)
+	private:
+		void GatherFilesToCopy(const FString& InSourceDirectoryPath, TArray<FString>& OutSourceDirectories, TArray<FString>& OutSourceFiles)
 		{
-			SourceFiles.Add(SourceDirectories[Idx] / SourceNames[NameIdx]);
+			const FString SourceDirectorySearchWildcard = InSourceDirectoryPath / TEXT("*");
+
+			OutSourceDirectories.Emplace(InSourceDirectoryPath);
+
+			TArray<FString> SourceFilenames;
+			IFileManager::Get().FindFiles(SourceFilenames, *SourceDirectorySearchWildcard, true, false);
+
+			OutSourceFiles.Reserve(OutSourceFiles.Num() + SourceFilenames.Num());
+			for (const FString& SourceFilename : SourceFilenames)
+			{
+				OutSourceFiles.Emplace(InSourceDirectoryPath / SourceFilename);
+			}
+
+			TArray<FString> SourceSubDirectoryNames;
+			IFileManager::Get().FindFiles(SourceSubDirectoryNames, *SourceDirectorySearchWildcard, false, true);
+
+			for (const FString& SourceSubDirectoryName : SourceSubDirectoryNames)
+			{
+				const FString SourceSubDirectoryPath = InSourceDirectoryPath / SourceSubDirectoryName;
+				if (!SourceDirectoriesToSkip.Contains(SourceSubDirectoryPath))
+				{
+					GatherFilesToCopy(SourceSubDirectoryPath, OutSourceDirectories, OutSourceFiles);
+				}
+			}
 		}
-	}
+
+		FString RootSourceDirectory;
+		TArray<FString> SourceDirectoriesToSkip;
+	};
+
+	TArray<FString> SourceDirectories;
+	TArray<FString> SourceFiles;
+	FGatherFilesToCopyHelper(OldDirectoryName).GatherFilesToCopy(SourceDirectories, SourceFiles);
 
 	// Copy everything
 	bool bCopySucceeded = true;
