Ensured that all files are deleted when new project creation fails

Prevent new projects from being created in folders with existing files so that we don't delete a user's files if they try to do so
Removed a couple of functions that were no longer being used

diff --git a/Engine/Source/Editor/GameProjectGeneration/Private/GameProjectUtils.cpp b/Engine/Source/Editor/GameProjectGeneration/Private/GameProjectUtils.cpp
--- a/Engine/Source/Editor/GameProjectGeneration/Private/GameProjectUtils.cpp
+++ b/Engine/Source/Editor/GameProjectGeneration/Private/GameProjectUtils.cpp
@@ -471,6 +471,15 @@ bool GameProjectUtils::IsValidProjectFileForCreation(const FString& ProjectFile,
 		return false;
 	}
 
+	// Don't allow any files within target directory so we can safely delete everything on failure
+	TArray<FString> ExistingFiles;
+	IFileManager::Get().FindFiles(ExistingFiles, *(FPaths::GetPath(ProjectFile) / TEXT("*")), true, true);
+	if (ExistingFiles.Num() > 0)
+	{
+		OutFailReason = LOCTEXT("ProjectFileCannotBeWithExistingFiles", "Project cannot be saved in a folder with existing files. Please choose a different directory/project name.");
+		return false;
+	}
+
 	return true;
 }
 
@@ -633,6 +642,17 @@ bool GameProjectUtils::CreateProject(const FProjectInformation& InProjectInfo, F
 		TemplateName = FPaths::GetBaseFilename(InProjectInfo.TemplateFile);
 	}
 
+	if (!bProjectCreationSuccessful && CleanupIsEnabled())
+	{
+		// Delete the new project folder
+		const FString NewProjectFolder = FPaths::GetPath(InProjectInfo.ProjectFilename);
+		IFileManager::Get().DeleteDirectory(*NewProjectFolder, /*RequireExists=*/false, /*Tree=*/true);
+		if( OutCreatedFiles != nullptr )
+		{
+			OutCreatedFiles->Empty();
+		}
+	}
+
 	if( FEngineAnalytics::IsAvailable() )
 	{
 		TArray<FAnalyticsEventAttribute> EventAttributes;
@@ -1067,14 +1087,12 @@ bool GameProjectUtils::GenerateProjectFromScratch(const FProjectInformation& InP
 	// Generate config files
 	if (!GenerateConfigFiles(InProjectInfo, CreatedFiles, OutFailReason))
 	{
-		DeleteCreatedFiles(NewProjectFolder, CreatedFiles);
 		return false;
 	}
 
 	// Insert any required feature packs (EG starter content) into ini file. These will be imported automatically when the editor is first run
 	if(!InsertFeaturePacksIntoINIFile(InProjectInfo, OutFailReason))
 	{
-		DeleteCreatedFiles(NewProjectFolder, CreatedFiles);
 		return false;
 	}
 	
@@ -1085,7 +1103,6 @@ bool GameProjectUtils::GenerateProjectFromScratch(const FProjectInformation& InP
 		FFormatNamedArguments Args;
 		Args.Add( TEXT("ContentFolder"), FText::FromString( ContentFolder ) );
 		OutFailReason = FText::Format( LOCTEXT("FailedToCreateContentFolder", "Failed to create the content folder {ContentFolder}"), Args );
-		DeleteCreatedFiles(NewProjectFolder, CreatedFiles);
 		return false;
 	}
 
@@ -1100,7 +1117,6 @@ bool GameProjectUtils::GenerateProjectFromScratch(const FProjectInformation& InP
 		// Generate basic source code files
 		if ( !GenerateBasicSourceCode(NewProjectFolder / TEXT("Source"), NewProjectName, NewProjectFolder, StartupModuleNames, CreatedFiles, OutFailReason) )
 		{
-			DeleteCreatedFiles(NewProjectFolder, CreatedFiles);
 			return false;
 		}
 
@@ -1108,7 +1124,6 @@ bool GameProjectUtils::GenerateProjectFromScratch(const FProjectInformation& InP
 		// Generate game framework source code files
 		if ( !GenerateGameFrameworkSourceCode(NewProjectFolder / TEXT("Source"), NewProjectName, CreatedFiles, OutFailReason) )
 		{
-			DeleteCreatedFiles(NewProjectFolder, CreatedFiles);
 			return false;
 		}
 	}
@@ -1129,7 +1144,6 @@ bool GameProjectUtils::GenerateProjectFromScratch(const FProjectInformation& InP
 		if(!Descriptor.Save(InProjectInfo.ProjectFilename, LocalFailReason))
 		{
 			OutFailReason = LocalFailReason;
-			DeleteCreatedFiles(NewProjectFolder, CreatedFiles);
 			return false;
 		}
 		CreatedFiles.Add(InProjectInfo.ProjectFilename);
@@ -1137,7 +1151,6 @@ bool GameProjectUtils::GenerateProjectFromScratch(const FProjectInformation& InP
 		// Set the engine identifier for it. Do this after saving, so it can be correctly detected as foreign or non-foreign.
 		if(!SetEngineAssociationForForeignProject(InProjectInfo.ProjectFilename, OutFailReason))
 		{
-			DeleteCreatedFiles(NewProjectFolder, CreatedFiles);
 			return false;
 		}
 	}
@@ -1149,8 +1162,6 @@ bool GameProjectUtils::GenerateProjectFromScratch(const FProjectInformation& InP
 		// Generate project files
 		if ( !GenerateCodeProjectFiles(InProjectInfo.ProjectFilename, OutFailReason, OutFailLog) )
 		{
-			DeleteGeneratedProjectFiles(InProjectInfo.ProjectFilename);
-			DeleteCreatedFiles(NewProjectFolder, CreatedFiles);
 			return false;
 		}
 	}
@@ -1323,7 +1334,6 @@ bool GameProjectUtils::CreateProjectFromTemplate(const FProjectInformation& InPr
 				FailArgs.Add(TEXT("SrcFilename"), FText::FromString(SrcFilename));
 				FailArgs.Add(TEXT("DestFilename"), FText::FromString(DestFilename));
 				OutFailReason = FText::Format(LOCTEXT("FailedToCopyFile", "Failed to copy \"{SrcFilename}\" to \"{DestFilename}\"."), FailArgs);
-				DeleteCreatedFiles(DestFolder, CreatedFiles);
 				return false;
 			}
 		}
@@ -1365,7 +1375,6 @@ bool GameProjectUtils::CreateProjectFromTemplate(const FProjectInformation& InPr
 				FFormatNamedArguments Args;
 				Args.Add( TEXT("FileToFix"), FText::FromString( FileToFix ) );
 				OutFailReason = FText::Format( LOCTEXT("FailedToFixUpFile", "Failed to process file \"{FileToFix}\"."), Args );
-				DeleteCreatedFiles(DestFolder, CreatedFiles);
 				return false;
 			}
 		}
@@ -1507,7 +1516,6 @@ bool GameProjectUtils::CreateProjectFromTemplate(const FProjectInformation& InPr
 		if ( !bSuccessfullyProcessed )
 		{
 			OutFailReason = LOCTEXT("FailedToFixUpDefaultEngine", "Failed to process file DefaultEngine.ini");
-			DeleteCreatedFiles(DestFolder, CreatedFiles);
 			return false;
 		}
 	}
@@ -1515,13 +1523,11 @@ bool GameProjectUtils::CreateProjectFromTemplate(const FProjectInformation& InPr
 	// Insert any required feature packs (EG starter content) into ini file. These will be imported automatically when the editor is first run
 	if(!InsertFeaturePacksIntoINIFile(InProjectInfo, OutFailReason))
 	{
-		DeleteCreatedFiles(DestFolder, CreatedFiles);
 		return false;
 	}
 
 	if( AddSharedContentToProject(InProjectInfo, CreatedFiles, OutFailReason ) == false )
 	{
-		DeleteCreatedFiles(DestFolder, CreatedFiles);
 		return false;
 	}
 
@@ -1534,7 +1540,6 @@ bool GameProjectUtils::CreateProjectFromTemplate(const FProjectInformation& InPr
 		FProjectDescriptor Project;
 		if(!Project.Load(InProjectInfo.TemplateFile, OutFailReason))
 		{
-			DeleteCreatedFiles(DestFolder, CreatedFiles);
 			return false;
 		}
 
@@ -1554,14 +1559,12 @@ bool GameProjectUtils::CreateProjectFromTemplate(const FProjectInformation& InPr
 		// Save it to disk
 		if(!Project.Save(InProjectInfo.ProjectFilename, OutFailReason))
 		{
-			DeleteCreatedFiles(DestFolder, CreatedFiles);
 			return false;
 		}
 
 		// Set the engine identifier if it's a foreign project. Do this after saving, so it can be correctly detected as foreign.
 		if(!SetEngineAssociationForForeignProject(InProjectInfo.ProjectFilename, OutFailReason))
 		{
-			DeleteCreatedFiles(DestFolder, CreatedFiles);
 			return false;
 		}
 
@@ -1577,8 +1580,6 @@ bool GameProjectUtils::CreateProjectFromTemplate(const FProjectInformation& InPr
 		// Generate project files
 		if ( !GenerateCodeProjectFiles(InProjectInfo.ProjectFilename, OutFailReason, OutFailLog) )
 		{
-			DeleteGeneratedProjectFiles(InProjectInfo.ProjectFilename);
-			DeleteCreatedFiles(DestFolder, CreatedFiles);
 			return false;
 		}
 	}
@@ -1587,8 +1588,6 @@ bool GameProjectUtils::CreateProjectFromTemplate(const FProjectInformation& InPr
 
 	if (!TemplateDefs->PostGenerateProject(DestFolder, SrcFolder, InProjectInfo.ProjectFilename, InProjectInfo.TemplateFile, InProjectInfo.bShouldGenerateCode, OutFailReason))
 	{
-		DeleteGeneratedProjectFiles(InProjectInfo.ProjectFilename);
-		DeleteCreatedFiles(DestFolder, CreatedFiles);
 		return false;
 	}
 	
@@ -1707,35 +1706,6 @@ void GameProjectUtils::DeleteCreatedFiles(const FString& RootFolder, const TArra
 	}
 }
 
-void GameProjectUtils::DeleteGeneratedProjectFiles(const FString& NewProjectFile)
-{
-	if (CleanupIsEnabled())
-	{
-		const FString NewProjectFolder = FPaths::GetPath(NewProjectFile);
-		const FString NewProjectName = FPaths::GetBaseFilename(NewProjectFile);
-
-		// Since it is hard to tell which files were created from the code project file generation process, just delete the entire ProjectFiles folder.
-		const FString IntermediateProjectFileFolder = NewProjectFolder / TEXT("Intermediate") / TEXT("ProjectFiles");
-		IFileManager::Get().DeleteDirectory(*IntermediateProjectFileFolder, /*RequireExists=*/false, /*Tree=*/true);
-
-		// Delete the solution file
-		const FString SolutionFileName = NewProjectFolder / NewProjectName + TEXT(".sln");
-		IFileManager::Get().Delete( *SolutionFileName );
-	}
-}
-
-void GameProjectUtils::DeleteGeneratedBuildFiles(const FString& NewProjectFolder)
-{
-	if (CleanupIsEnabled())
-	{
-		// Since it is hard to tell which files were created from the build process, just delete the entire Binaries and Build folders.
-		const FString BinariesFolder = NewProjectFolder / TEXT("Binaries");
-		const FString BuildFolder    = NewProjectFolder / TEXT("Intermediate") / TEXT("Build");
-		IFileManager::Get().DeleteDirectory(*BinariesFolder, /*RequireExists=*/false, /*Tree=*/true);
-		IFileManager::Get().DeleteDirectory(*BuildFolder, /*RequireExists=*/false, /*Tree=*/true);
-	}
-}
-
 FString GameProjectUtils::GetHardwareConfigString(const FProjectInformation& InProjectInfo)
 {
 	FString HardwareTargeting;
diff --git a/Engine/Source/Editor/GameProjectGeneration/Public/GameProjectUtils.h b/Engine/Source/Editor/GameProjectGeneration/Public/GameProjectUtils.h
--- a/Engine/Source/Editor/GameProjectGeneration/Public/GameProjectUtils.h
+++ b/Engine/Source/Editor/GameProjectGeneration/Public/GameProjectUtils.h
@@ -285,12 +285,6 @@ private:
 	/** Returns true if file cleanup on failure is enabled, false if not */
 	static bool CleanupIsEnabled();
 
-	/** Deletes any files that were generated by the generate project files step */
-	static void DeleteGeneratedProjectFiles(const FString& NewProjectFile);
-
-	/** Deletes any files that were generated by the build step */
-	static void DeleteGeneratedBuildFiles(const FString& NewProjectFolder);
-
 	/** Creates ini files for a new project. On failure, OutFailReason will be populated. */
 	static bool GenerateConfigFiles(const FProjectInformation& InProjectInfo, TArray<FString>& OutCreatedFiles, FText& OutFailReason);
 
