Fix non-unity analytics compile errors

diff --git a/Engine/Source/Editor/GameProjectGeneration/Private/SGetSuggestedIDEWidget.cpp b/Engine/Source/Editor/GameProjectGeneration/Private/SGetSuggestedIDEWidget.cpp
--- a/Engine/Source/Editor/GameProjectGeneration/Private/SGetSuggestedIDEWidget.cpp
+++ b/Engine/Source/Editor/GameProjectGeneration/Private/SGetSuggestedIDEWidget.cpp
@@ -6,6 +6,8 @@
 #include "SNotificationList.h"
 #include "NotificationManager.h"
 #include "SHyperlink.h"
+#include "EngineAnalytics.h"
+#include "IAnalyticsProvider.h"
 
 #define LOCTEXT_NAMESPACE "GameProjectGeneration"
 
diff --git a/Engine/Source/Editor/GameProjectGeneration/Private/SProjectBrowser.cpp b/Engine/Source/Editor/GameProjectGeneration/Private/SProjectBrowser.cpp
--- a/Engine/Source/Editor/GameProjectGeneration/Private/SProjectBrowser.cpp
+++ b/Engine/Source/Editor/GameProjectGeneration/Private/SProjectBrowser.cpp
@@ -11,6 +11,9 @@
 #include "PlatformInfo.h"
 #include "SSearchBox.h"
 #include "Settings/EditorSettings.h"
+#include "AnalyticsEventAttribute.h"
+#include "EngineAnalytics.h"
+#include "IAnalyticsProvider.h"
 
 #define LOCTEXT_NAMESPACE "ProjectBrowser"
 
