Fix non-unity build.

diff --git a/Engine/Source/Editor/GameProjectGeneration/Private/SNewClassDialog.cpp b/Engine/Source/Editor/GameProjectGeneration/Private/SNewClassDialog.cpp
--- a/Engine/Source/Editor/GameProjectGeneration/Private/SNewClassDialog.cpp
+++ b/Engine/Source/Editor/GameProjectGeneration/Private/SNewClassDialog.cpp
@@ -31,6 +31,7 @@
 #include "SNotificationList.h"
 #include "NotificationManager.h"
 #include "Engine/BlueprintGeneratedClass.h"
+#include "FeaturedClasses.inl"
 
 
 #define LOCTEXT_NAMESPACE "GameProjectGeneration"
diff --git a/Engine/Source/Editor/GameProjectGeneration/Public/FeaturedClasses.inl b/Engine/Source/Editor/GameProjectGeneration/Public/FeaturedClasses.inl
--- a/Engine/Source/Editor/GameProjectGeneration/Public/FeaturedClasses.inl
+++ b/Engine/Source/Editor/GameProjectGeneration/Public/FeaturedClasses.inl
@@ -1,5 +1,7 @@
 // Copyright 1998-2016 Epic Games, Inc. All Rights Reserved.
 
+#pragma once
+
 /** Inline file to avoid introducing many UObject headers into the global namespace */
 
 #include "Kismet/BlueprintFunctionLibrary.h"
