Copying //UE4/Dev-Core to //UE4/Dev-Main (Source: //UE4/Dev-Core @ 3283640)

diff --git a/Engine/Plugins/2D/Paper2D/Config/BasePaper2D.ini b/Engine/Plugins/2D/Paper2D/Config/BasePaper2D.ini
new file mode 100644
--- /dev/null
+++ b/Engine/Plugins/2D/Paper2D/Config/BasePaper2D.ini
@@ -0,0 +1,10 @@
+[CoreRedirects]
+
++ClassRedirects=(OldName="PaperAnimatedRenderComponent", NewName="/Script/Paper2D.PaperFlipbookComponent")
++ClassRedirects=(OldName="PaperRenderActor", NewName="/Script/Paper2D.PaperSpriteActor")
++ClassRedirects=(OldName="PaperRenderComponent", NewName="/Script/Paper2D.PaperSpriteComponent")
++ClassRedirects=(OldName="PaperTileMapRenderComponent", NewName="/Script/Paper2D.PaperTileMapComponent")
++ClassRedirects=(OldName="PaperSpriteSheet", NewName="/Script/PaperSpritesheetImporter.PaperSpriteSheet")
+	
++PropertyRedirects=(OldName="PaperTileSet.ExperimentalPerTileData", NewName="PaperTileSet.PerTileData")
++PropertyRedirects=(OldName="PaperRenderComponent.TestMaterial", NewName="PaperRenderComponent.MaterialOverride")
diff --git a/Engine/Plugins/Online/OnlineSubsystem/Source/Private/OnlineKeyValuePair.cpp b/Engine/Plugins/Online/OnlineSubsystem/Source/Private/OnlineKeyValuePair.cpp
--- a/Engine/Plugins/Online/OnlineSubsystem/Source/Private/OnlineKeyValuePair.cpp
+++ b/Engine/Plugins/Online/OnlineSubsystem/Source/Private/OnlineKeyValuePair.cpp
@@ -1147,7 +1147,7 @@ bool FVariantDataConverter::ConvertScalarUPropertyToVariant(UProperty* Property,
 	{
 		// export enums as strings
 		UEnum* EnumDef = EnumProperty->GetEnum();
-		FString StringValue = EnumDef->GetEnumName(EnumProperty->GetUnderlyingProperty()->GetSignedIntPropertyValue(Value));
+		FString StringValue = EnumDef->GetEnumNameStringByValue(EnumProperty->GetUnderlyingProperty()->GetSignedIntPropertyValue(Value));
 		OutVariantData.SetValue(StringValue);
 	}
 	else if (UNumericProperty *NumericProperty = Cast<UNumericProperty>(Property))
@@ -1157,7 +1157,7 @@ bool FVariantDataConverter::ConvertScalarUPropertyToVariant(UProperty* Property,
 		if (EnumDef != NULL)
 		{
 			// export enums as strings
-			FString StringValue = EnumDef->GetEnumName(NumericProperty->GetSignedIntPropertyValue(Value));
+			FString StringValue = EnumDef->GetEnumNameStringByValue(NumericProperty->GetSignedIntPropertyValue(Value));
 			OutVariantData.SetValue(StringValue);
 		}
 		// We want to export numbers as numbers
diff --git a/Engine/Plugins/Runtime/GameplayAbilities/Config/BaseGameplayAbilities.ini b/Engine/Plugins/Runtime/GameplayAbilities/Config/BaseGameplayAbilities.ini
new file mode 100644
--- /dev/null
+++ b/Engine/Plugins/Runtime/GameplayAbilities/Config/BaseGameplayAbilities.ini
@@ -0,0 +1,11 @@
+[CoreRedirects]
+
++ClassRedirects=(OldName="GameplayAbilityBlueprintGeneratedClass",NewName="/Script/Engine.BlueprintGeneratedClass")
+
++FunctionRedirects=(OldName="AbilitySystemComponent.TryActivateAbilityByTag",NewName="AbilitySystemComponent.TryActivateAbilitiesByTag")
++FunctionRedirects=(OldName="AbilitySystemComponent.GetEffectContext",NewName="AbilitySystemComponent.MakeEffectContext")
+
++PropertyRedirects=(OldName="GameplayEffect.InheritableClearTagsContainer",NewName="GameplayEffect.RemoveGameplayEffectsWithTags")
++PropertyRedirects=(OldName="GameplayEffect.IsInhibited",NewName="GameplayEffect.bIsInhibited")
+
++EnumRedirects=(OldName="EGameplayAbilityNetExecutionPolicy",ValueChanges=(("EGameplayAbilityNetExecutionPolicy::Client","EGameplayAbilityNetExecutionPolicy::LocalOnly"),("EGameplayAbilityNetExecutionPolicy::Predictive","EGameplayAbilityNetExecutionPolicy::LocalPredicted"),("EGameplayAbilityNetExecutionPolicy::Server","EGameplayAbilityNetExecutionPolicy::ServerInitiated")) )
diff --git a/Engine/Plugins/Runtime/PhysXVehicles/Config/BasePhysXVehicles.ini b/Engine/Plugins/Runtime/PhysXVehicles/Config/BasePhysXVehicles.ini
new file mode 100644
--- /dev/null
+++ b/Engine/Plugins/Runtime/PhysXVehicles/Config/BasePhysXVehicles.ini
@@ -0,0 +1,15 @@
+[CoreRedirects]
+
++ClassRedirects=(OldName="AnimGraphNode_WheelHandler", NewName="/Script/PhysXVehiclesEditor.AnimGraphNode_WheelHandler")
++ClassRedirects=(OldName="SimpleWheeledVehicleMovementComponent", NewName="/Script/PhysXVehicles.SimpleWheeledVehicleMovementComponent")
++ClassRedirects=(OldName="VehicleSim", NewName="/Script/PhysXVehicles.WheeledVehicleMovementComponent")
++ClassRedirects=(OldName="MovementComp_Vehicle", NewName="/Script/PhysXVehicles.WheeledVehicleMovementComponent")
++ClassRedirects=(OldName="Vehicle", NewName="/Script/PhysXVehicles.WheeledVehicle")
++ClassRedirects=(OldName="WheeledVehicle", NewName="/Script/PhysXVehicles.WheeledVehicle")
++ClassRedirects=(OldName="VehicleMovementComponent", NewName="/Script/PhysXVehicles.WheeledVehicleMovementComponent")
++ClassRedirects=(OldName="WheeledVehicleMovementComponent", NewName="/Script/PhysXVehicles.WheeledVehicleMovementComponent")
++ClassRedirects=(OldName="VehicleMovementComponent4W", NewName="/Script/PhysXVehicles.WheeledVehicleMovementComponent4W")
++ClassRedirects=(OldName="WheeledVehicleMovementComponent4W", NewName="/Script/PhysXVehicles.WheeledVehicleMovementComponent4W")
++ClassRedirects=(OldName="VehicleAnimInstance", NewName="/Script/PhysXVehicles.VehicleAnimInstance")
++ClassRedirects=(OldName="VehicleWheel", NewName="/Script/PhysXVehicles.VehicleWheel")
++StructRedirects=(OldName="AnimNode_WheelHandler", NewName="/Script/PhysXVehicles.AnimNode_WheelHandler")
diff --git a/Engine/Source/Developer/HotReload/Private/HotReload.cpp b/Engine/Source/Developer/HotReload/Private/HotReload.cpp
--- a/Engine/Source/Developer/HotReload/Private/HotReload.cpp
+++ b/Engine/Source/Developer/HotReload/Private/HotReload.cpp
@@ -1237,6 +1237,13 @@ void FHotReloadModule::ReinstanceClasses()
 	TMap<UClass*, UClass*> OldToNewClassesMap;
 	for (const TPair<UClass*, UClass*>& Pair : ClassesToReinstance)
 	{
+		// Don't allow reinstancing of UEngine classes
+		if (Pair.Key->IsChildOf(UEngine::StaticClass()))
+		{
+			UE_LOG(LogHotReload, Warning, TEXT("Engine class '%s' has changed but will be ignored for hot reload"), *Pair.Key->GetName());
+			continue;
+		}
+
 		if (Pair.Value != nullptr)
 		{
 			OldToNewClassesMap.Add(Pair.Key, Pair.Value);
@@ -1245,7 +1252,11 @@ void FHotReloadModule::ReinstanceClasses()
 
 	for (const TPair<UClass*, UClass*>& Pair : ClassesToReinstance)
 	{
-		ReinstanceClass(Pair.Key, Pair.Value, OldToNewClassesMap);
+		// Don't allow reinstancing of UEngine classes
+		if (!Pair.Key->IsChildOf(UEngine::StaticClass()))
+		{
+			ReinstanceClass(Pair.Key, Pair.Value, OldToNewClassesMap);
+		}
 	}
 
 	ClassesToReinstance.Empty();
diff --git a/Engine/Source/Developer/HotReload/Private/HotReloadClassReinstancer.cpp b/Engine/Source/Developer/HotReload/Private/HotReloadClassReinstancer.cpp
--- a/Engine/Source/Developer/HotReload/Private/HotReloadClassReinstancer.cpp
+++ b/Engine/Source/Developer/HotReload/Private/HotReloadClassReinstancer.cpp
@@ -20,7 +20,6 @@ void FHotReloadClassReinstancer::SetupNewClassReinstancing(UClass* InNewClass, U
 	DuplicatedClass = InOldClass;
 	OriginalCDO = InOldClass->GetDefaultObject();
 	bHasReinstanced = false;
-	bSkipGarbageCollection = false;
 	bNeedsReinstancing = true;
 	NewClass = InNewClass;
 
@@ -223,7 +222,6 @@ void FHotReloadClassReinstancer::RecreateCDOAndSetupOldClassReinstancing(UClass*
 	DuplicatedClass = InOldClass;
 	OriginalCDO = InOldClass->GetDefaultObject();
 	bHasReinstanced = false;
-	bSkipGarbageCollection = false;
 	bNeedsReinstancing = false;
 	NewClass = InOldClass; // The class doesn't change in this case
 
diff --git a/Engine/Source/Developer/ShaderFormatVectorVM/Private/ir_vm_gen_bytecode_visitor.cpp b/Engine/Source/Developer/ShaderFormatVectorVM/Private/ir_vm_gen_bytecode_visitor.cpp
--- a/Engine/Source/Developer/ShaderFormatVectorVM/Private/ir_vm_gen_bytecode_visitor.cpp
+++ b/Engine/Source/Developer/ShaderFormatVectorVM/Private/ir_vm_gen_bytecode_visitor.cpp
@@ -676,8 +676,7 @@ struct op_input : public op_base
 	{
 		if (instance_idx_component)
 		{
-			FString Fmt = TEXT("[R][%d] = InputData(%d, %d, [%s][%d]);\n");
-			return FString::Printf(*Fmt,
+			return FString::Printf(TEXT("[R][%d] = InputData(%d, %d, [%s][%d]);\n"),
 				dest_component->offset,
 				dataset_idx,
 				register_idx,
@@ -686,8 +685,7 @@ struct op_input : public op_base
 		}
 		else
 		{
-			FString Fmt = TEXT("[R][%d] = InputData(%d, %d);\n");
-			return FString::Printf(*Fmt, dest_component->offset, dataset_idx, register_idx);
+			return FString::Printf(TEXT("[R][%d] = InputData(%d, %d);\n"), dest_component->offset, dataset_idx, register_idx);
 		}
 	}
 
@@ -720,8 +718,7 @@ struct op_input_noadvance : public op_input
 	{
 		if (instance_idx_component)
 		{
-			FString Fmt = TEXT("[R][%d] = InputDataNoadvance(%d, %d, [%s][%d]);\n");
-			return FString::Printf(*Fmt,
+			return FString::Printf(TEXT("[R][%d] = InputDataNoadvance(%d, %d, [%s][%d]);\n"),
 				dest_component->offset,
 				dataset_idx,
 				register_idx,
@@ -730,8 +727,7 @@ struct op_input_noadvance : public op_input
 		}
 		else
 		{
-			FString Fmt = TEXT("[R][%d] = InputDataNoadvance(%d, %d);\n");
-			return FString::Printf(*Fmt, dest_component->offset, dataset_idx, register_idx);
+			return FString::Printf(TEXT("[R][%d] = InputDataNoadvance(%d, %d);\n"), dest_component->offset, dataset_idx, register_idx);
 		}
 	}
 
@@ -771,8 +767,7 @@ struct op_output : public op_base
 		check(instance_idx_component);
 		check(value_component);
 
-		FString Fmt = TEXT("OutputData(%d, [%s][%d], [%s][%d]);\n");
-		return FString::Printf(*Fmt,
+		return FString::Printf(TEXT("OutputData(%d, [%s][%d], [%s][%d]);\n"),
 			register_idx,
 			get_location_string(instance_idx_component->owner->location),
 			instance_idx_component->offset,
@@ -812,8 +807,7 @@ struct op_index_acquire : public op_base
 	{
 		check(valid_component);
 
-		FString Fmt = TEXT("[R][%d] = AcquireIndex(%d, [%s][%d]);\n");
-		return FString::Printf(*Fmt,
+		return FString::Printf(TEXT("[R][%d] = AcquireIndex(%d, [%s][%d]);\n"),
 			dest_component->offset,
 			dataset_idx,
 			get_location_string(valid_component->owner->location),
diff --git a/Engine/Source/Editor/BlueprintGraph/Classes/K2Node.h b/Engine/Source/Editor/BlueprintGraph/Classes/K2Node.h
--- a/Engine/Source/Editor/BlueprintGraph/Classes/K2Node.h
+++ b/Engine/Source/Editor/BlueprintGraph/Classes/K2Node.h
@@ -362,7 +362,7 @@ protected:
 		ERedirectType_None,
 		ERedirectType_Name,
 		ERedirectType_Value,
-		ERedirectType_Custom
+		ERedirectType_Custom // DEPRECATED
 	};
 
 	// Handles the actual reconstruction (copying data, links, name, etc...) from two pins that have already been matched together
@@ -374,7 +374,7 @@ protected:
 	// Helper function to properly destroy a set of pins
 	BLUEPRINTGRAPH_API void DestroyPinList(TArray<UEdGraphPin*>& InPins);
 
-	/** Allows the custom transformation of a param's value when redirecting a matched pin; called only when DoPinsMatchForReconstruction returns ERedirectType_Custom **/
+	DEPRECATED(4.16, "This function is deprecated, override ReconstructSinglePin instead ")
 	BLUEPRINTGRAPH_API virtual void CustomMapParamValue(UEdGraphPin& Pin);
 
 	/** Whether or not two pins match for purposes of reconnection after reconstruction.  This allows pins that may have had their names changed via reconstruction to be matched to their old values on a node-by-node basis, if needed*/
diff --git a/Engine/Source/Editor/BlueprintGraph/Classes/K2Node_CreateDelegate.h b/Engine/Source/Editor/BlueprintGraph/Classes/K2Node_CreateDelegate.h
--- a/Engine/Source/Editor/BlueprintGraph/Classes/K2Node_CreateDelegate.h
+++ b/Engine/Source/Editor/BlueprintGraph/Classes/K2Node_CreateDelegate.h
@@ -40,6 +40,7 @@ public:
 	virtual void GetMenuActions(FBlueprintActionDatabaseRegistrar& ActionRegistrar) const override;
 	virtual FText GetMenuCategory() const override;
 	virtual int32 GetNodeRefreshPriority() const override { return EBaseNodeRefreshPriority::Low_ReceivesDelegateSignature; }
+	virtual ERedirectType DoPinsMatchForReconstruction(const UEdGraphPin* NewPin, int32 NewPinIndex, const UEdGraphPin* OldPin, int32 OldPinIndex) const override;
 	// End of UK2Node interface
 
 	bool IsValid(FString* OutMsg = NULL, bool bDontUseSkeletalClassForSelf = false) const;
diff --git a/Engine/Source/Editor/BlueprintGraph/Private/CallFunctionHandler.cpp b/Engine/Source/Editor/BlueprintGraph/Private/CallFunctionHandler.cpp
--- a/Engine/Source/Editor/BlueprintGraph/Private/CallFunctionHandler.cpp
+++ b/Engine/Source/Editor/BlueprintGraph/Private/CallFunctionHandler.cpp
@@ -155,8 +155,8 @@ void FKCHandler_CallFunction::CreateFunctionCallStatement(FKismetFunctionContext
 					}
 					else
 					{
-						CompilerContext.MessageLog.Error(*FString::Printf(*LOCTEXT("FindPinFromLinkage_Error", "Function %s (called from @@) was specified with LatentInfo metadata but does not have a pin named %s").ToString(),
-							*(Function->GetName()), *(It.Value())), Node);
+						CompilerContext.MessageLog.Error(*FText::Format(LOCTEXT("FindPinFromLinkage_ErrorFmt", "Function {0} (called from @@) was specified with LatentInfo metadata but does not have a pin named {1}"),
+							FText::FromString(Function->GetName()), FText::FromString(It.Value())).ToString(), Node);
 					}
 				}
 			}
@@ -314,7 +314,7 @@ void FKCHandler_CallFunction::CreateFunctionCallStatement(FKismetFunctionContext
 
 			if (!bFoundParam)
 			{
-				CompilerContext.MessageLog.Error(*FString::Printf(*LOCTEXT("FindPinParameter_Error", "Could not find a pin for the parameter %s of %s on @@").ToString(), *Property->GetName(), *Function->GetName()), Node);
+				CompilerContext.MessageLog.Error(*FText::Format(LOCTEXT("FindPinParameter_ErrorFmt", "Could not find a pin for the parameter {0} of {1} on @@"), FText::FromString(Property->GetName()), FText::FromString(Function->GetName())).ToString(), Node);
 				bMatchedAllParams = false;
 			}
 		}
@@ -322,7 +322,7 @@ void FKCHandler_CallFunction::CreateFunctionCallStatement(FKismetFunctionContext
 		// At this point, we should have consumed all pins.  If not, there are extras that need to be removed.
 		for (int32 i = 0; i < RemainingPins.Num(); ++i)
 		{
-			CompilerContext.MessageLog.Error(*FString::Printf(*LOCTEXT("PinMismatchParameter_Error", "Pin @@ named %s doesn't match any parameters of function %s").ToString(), *RemainingPins[i]->PinName, *Function->GetName()), RemainingPins[i]);
+			CompilerContext.MessageLog.Error(*FText::Format(LOCTEXT("PinMismatchParameter_ErrorFmt", "Pin @@ named {0} doesn't match any parameters of function {1}"), FText::FromString(RemainingPins[i]->PinName), FText::FromString(Function->GetName())).ToString(), RemainingPins[i]);
 		}
 
 		if (NumErrorsAtStart == CompilerContext.MessageLog.NumErrors)
@@ -522,7 +522,7 @@ void FKCHandler_CallFunction::CreateFunctionCallStatement(FKismetFunctionContext
 	}
 	else
 	{
-		FString WarningMessage = FString::Printf(*LOCTEXT("FindFunction_Error", "Could not find the function '%s' called from @@").ToString(), *GetFunctionNameFromNode(Node));
+		FString WarningMessage = FText::Format(LOCTEXT("FindFunction_ErrorFmt", "Could not find the function '{0}' called from @@"), FText::FromString(GetFunctionNameFromNode(Node))).ToString();
 		CompilerContext.MessageLog.Warning(*WarningMessage, Node);
 	}
 }
@@ -763,12 +763,12 @@ void FKCHandler_CallFunction::Compile(FKismetFunctionContext& Context, UEdGraphN
 		UEdGraphPin* ExecTriggeringPin = CompilerContext.GetSchema()->FindExecutionPin(*Node, EGPD_Input);
 		if (ExecTriggeringPin == NULL)
 		{
-			CompilerContext.MessageLog.Error(*FString::Printf(*NSLOCTEXT("KismetCompiler", "NoValidExecutionPinForCallFunc_Error", "@@ must have a valid execution pin").ToString()), Node);
+			CompilerContext.MessageLog.Error(*NSLOCTEXT("KismetCompiler", "NoValidExecutionPinForCallFunc_Error", "@@ must have a valid execution pin").ToString(), Node);
 			return;
 		}
 		else if (ExecTriggeringPin->LinkedTo.Num() == 0)
 		{
-			CompilerContext.MessageLog.Warning(*FString::Printf(*NSLOCTEXT("KismetCompiler", "NodeNeverExecuted_Warning", "@@ will never be executed").ToString()), Node);
+			CompilerContext.MessageLog.Warning(*NSLOCTEXT("KismetCompiler", "NodeNeverExecuted_Warning", "@@ will never be executed").ToString(), Node);
 			return;
 		}
 	}
@@ -798,7 +798,7 @@ void FKCHandler_CallFunction::CheckIfFunctionIsCallable(UFunction* Function, FKi
 	{
 		if (!IsCalledFunctionFinal(Node) && Function->GetName().Find(CompilerContext.GetSchema()->FN_ExecuteUbergraphBase.ToString()))
 		{
-			CompilerContext.MessageLog.Error(*FString::Printf(*NSLOCTEXT("KismetCompiler", "ShouldNotCallFromBlueprint_Error", "Function '%s' called from @@ should not be called from a Blueprint").ToString(), *Function->GetName()), Node);
+			CompilerContext.MessageLog.Error(*FText::Format(NSLOCTEXT("KismetCompiler", "ShouldNotCallFromBlueprint_ErrorFmt", "Function '{0}' called from @@ should not be called from a Blueprint"), FText::FromString(Function->GetName())).ToString(), Node);
 		}
 	}
 }
diff --git a/Engine/Source/Editor/BlueprintGraph/Private/DelegateNodeHandlers.cpp b/Engine/Source/Editor/BlueprintGraph/Private/DelegateNodeHandlers.cpp
--- a/Engine/Source/Editor/BlueprintGraph/Private/DelegateNodeHandlers.cpp
+++ b/Engine/Source/Editor/BlueprintGraph/Private/DelegateNodeHandlers.cpp
@@ -53,8 +53,7 @@ struct FKCHandlerDelegateHelper
 				FString const OwnerName = PropertyOwnerClass->GetName();
 				FString const PropName = DelegateNode->GetPropertyName().ToString();
 
-				FText const ErrorFormat = LOCTEXT("DelegateNotFound", "Could not find an event-dispatcher named \"%s\" in '%s'.\nMake sure '%s' has been compiled for @@");
-				MessageLog.Error(*FString::Printf(*ErrorFormat.ToString(), *PropName, *OwnerName, *OwnerName), DelegateNode);
+				MessageLog.Error(*FText::Format(LOCTEXT("DelegateNotFoundFmt", "Could not find an event-dispatcher named \"{0}\" in '{1}'.\nMake sure '{2}' has been compiled for @@"), FText::FromString(PropName), FText::FromString(OwnerName), FText::FromString(OwnerName)).ToString(), DelegateNode);
 			}
 			return NULL;
 		}
@@ -208,9 +207,7 @@ void FKCHandler_CreateDelegate::RegisterNets(FKismetFunctionContext& Context, UE
 
 	if(!DelegateNode->GetDelegateSignature())
 	{
-		const FString ErrorStr = FString::Printf(
-			*LOCTEXT("NoDelegateFunction", "@@ : unable to determine expected signature - is the delegate pin connected?").ToString(),
-			*DelegateFunctionName.ToString());
+		const FString ErrorStr = LOCTEXT("NoDelegateFunction", "@@ : unable to determine expected signature - is the delegate pin connected?").ToString();
 		CompilerContext.MessageLog.Error(*ErrorStr, DelegateNode);
 		return;
 	}
diff --git a/Engine/Source/Editor/BlueprintGraph/Private/K2Node.cpp b/Engine/Source/Editor/BlueprintGraph/Private/K2Node.cpp
--- a/Engine/Source/Editor/BlueprintGraph/Private/K2Node.cpp
+++ b/Engine/Source/Editor/BlueprintGraph/Private/K2Node.cpp
@@ -3,6 +3,7 @@
 
 #include "K2Node.h"
 #include "UObject/UnrealType.h"
+#include "UObject/CoreRedirects.h"
 #include "EdGraph/EdGraphPin.h"
 #include "UObject/Interface.h"
 #include "Engine/Blueprint.h"
@@ -431,14 +432,14 @@ void UK2Node::ReconstructNode()
 		{
 			UEdGraphPin* OtherPin = LinkedToCopy[LinkIdx];
 			// If we are linked to a pin that its owner doesn't know about, break that link
-			if ((OtherPin == NULL) || !OtherPin->GetOwningNodeUnchecked() || !OtherPin->GetOwningNode()->Pins.Contains(OtherPin))
+			if ((OtherPin == nullptr) || !OtherPin->GetOwningNodeUnchecked() || !OtherPin->GetOwningNode()->Pins.Contains(OtherPin))
 			{
 				Pin->LinkedTo.Remove(OtherPin);
 			}
 
 			if (Blueprint->bIsRegeneratingOnLoad && Linker->UE4Ver() < VER_UE4_INJECT_BLUEPRINT_STRUCT_PIN_CONVERSION_NODES)
 			{
-				if ((Pin->PinType.PinCategory != UEdGraphSchema_K2::PC_Struct))
+				if (OtherPin == nullptr || (Pin->PinType.PinCategory != UEdGraphSchema_K2::PC_Struct))
 				{
 					continue;
 				}
@@ -493,32 +494,19 @@ void UK2Node::GetRedirectPinNames(const UEdGraphPin& Pin, TArray<FString>& Redir
 
 UK2Node::ERedirectType UK2Node::ShouldRedirectParam(const TArray<FString>& OldPinNames, FName& NewPinName, const UK2Node * NewPinNode) const 
 {
+	DECLARE_SCOPE_CYCLE_COUNTER(TEXT("UK2Node::ShouldRedirectParam"), STAT_LinkerLoad_ShouldRedirectParam, STATGROUP_LoadTimeVerbose);
+
 	if ( ensure(NewPinNode) )
 	{
-		const TMultiMap<UClass*, FParamRemapInfo>& ParamRedirectMap = FMemberReference::GetParamRedirectMap();
-
-		if ( ParamRedirectMap.Num() > 0 )
+		for (const FString& OldPinName : OldPinNames)
 		{
-			// convert TArray<FString> to TArray<FName>, faster to search
-			TArray<FName> OldPinFNames;
-			for (auto NameIter=OldPinNames.CreateConstIterator(); NameIter; ++NameIter)
-			{
-				const FString& Name = *NameIter;
-				OldPinFNames.AddUnique(*Name);
-			}
-
-			// go through for the NewPinNode
-			for(TMultiMap<UClass*, FParamRemapInfo>::TConstKeyIterator ParamIter(ParamRedirectMap, NewPinNode->GetClass()); ParamIter; ++ParamIter)
+			const FCoreRedirect* ValueRedirect = nullptr;
+			FCoreRedirectObjectName NewRedirectName;
+			
+			if (FCoreRedirects::RedirectNameAndValues(ECoreRedirectFlags::Type_Property, OldPinName, NewRedirectName, &ValueRedirect))
 			{
-				const FParamRemapInfo& ParamRemap = ParamIter.Value();
-
-				// if it has it, return true
-				if (OldPinFNames.Contains(ParamRemap.OldParam)
-					&& (ParamRemap.NodeTitle == NAME_None || ParamRemap.NodeTitle.ToString() == NewPinNode->GetNodeTitle(ENodeTitleType::FullTitle).ToString()))
-				{
-					NewPinName = ParamRemap.NewParam;
-					return (ParamRemap.bCustomValueMapping ? ERedirectType_Custom : (ParamRemap.ParamValueMap.Num() ? ERedirectType_Value : ERedirectType_Name));
-				}
+				NewPinName = NewRedirectName.ObjectName;
+				return (ValueRedirect ? ERedirectType_Value : ERedirectType_Name);
 			}
 		}
 	}
@@ -631,20 +619,13 @@ UK2Node::ERedirectType UK2Node::DoPinsMatchForReconstruction(const UEdGraphPin*
 				{
 					const UEdGraphPin* CurPin = ParentHierarchy[ParentIndex].Pin;
 					const UEdGraphPin* ParentPin = CurPin ? CurPin->ParentPin : nullptr;
-					const UObject* SubCategoryObject = ParentPin ? ParentPin->PinType.PinSubCategoryObject.Get() : nullptr;
+					UStruct* SubCategoryStruct = ParentPin ? Cast<UStruct>(ParentPin->PinType.PinSubCategoryObject.Get()) : nullptr;
+
+					FName RedirectedPinName = SubCategoryStruct ? UProperty::FindRedirectedPropertyName(SubCategoryStruct, FName(*ParentHierarchy[ParentIndex].PropertyName)) : NAME_None;
 
-					TMap<FName, FName>* StructRedirects = SubCategoryObject ? UStruct::TaggedPropertyRedirects.Find(SubCategoryObject->GetFName()) : nullptr;
-					if (StructRedirects)
+					if (RedirectedPinName != NAME_Name)
 					{
-						FName* PropertyRedirect = StructRedirects->Find(FName(*ParentHierarchy[ParentIndex].PropertyName));
-						if (PropertyRedirect)
-						{
-							NewPinNameStr += FString("_") + PropertyRedirect->ToString();
-						}
-						else
-						{
-							NewPinNameStr += FString("_") + ParentHierarchy[ParentIndex].PropertyName;
-						}
+						NewPinNameStr += FString("_") + RedirectedPinName.ToString();
 					}
 					else
 					{
@@ -667,6 +648,8 @@ void UK2Node::CustomMapParamValue(UEdGraphPin& Pin)
 
 void UK2Node::ReconstructSinglePin(UEdGraphPin* NewPin, UEdGraphPin* OldPin, ERedirectType RedirectType)
 {
+	DECLARE_SCOPE_CYCLE_COUNTER(TEXT("UK2Node::ReconstructSinglePin"), STAT_LinkerLoad_ReconstructSinglePin, STATGROUP_LoadTimeVerbose);
+
 	UBlueprint* Blueprint = GetBlueprint();
 
 	check(NewPin && OldPin);
@@ -681,23 +664,13 @@ void UK2Node::ReconstructSinglePin(UEdGraphPin* NewPin, UEdGraphPin* OldPin, ERe
 			TArray<FString> OldPinNames;
 			GetRedirectPinNames(*OldPin, OldPinNames);
 
-			// convert TArray<FString> to TArray<FName>, faster to search
-			TArray<FName> OldPinFNames;
-			for (auto NameIter=OldPinNames.CreateConstIterator(); NameIter; ++NameIter)
-			{
-				const FString& Name = *NameIter;
-				OldPinFNames.AddUnique(*Name);
-			}
-
-			// go through for the NewPinNode
-			for(TMultiMap<UClass*, FParamRemapInfo>::TConstKeyIterator ParamIter(FMemberReference::GetParamRedirectMap(), Cast<UK2Node>(NewPin->GetOwningNode())->GetClass()); ParamIter; ++ParamIter)
+			for (const FString& OldPinName : OldPinNames)
 			{
-				const FParamRemapInfo& ParamRemap = ParamIter.Value();
+				const TMap<FString, FString>* ValueChanges = FCoreRedirects::GetValueRedirects(ECoreRedirectFlags::Type_Property, OldPinName);
 
-				// once we find it, see about remapping the value
-				if (OldPinFNames.Contains(ParamRemap.OldParam))
+				if (ValueChanges)
 				{
-					const FString* NewValue = ParamRemap.ParamValueMap.Find(NewPin->DefaultValue);
+					const FString* NewValue = ValueChanges->Find(*NewPin->DefaultValue);
 					if (NewValue)
 					{
 						NewPin->DefaultValue = *NewValue;
@@ -708,7 +681,9 @@ void UK2Node::ReconstructSinglePin(UEdGraphPin* NewPin, UEdGraphPin* OldPin, ERe
 		}
 		else if (RedirectType == ERedirectType_Custom)
 		{
+			PRAGMA_DISABLE_DEPRECATION_WARNINGS
 			CustomMapParamValue(*NewPin);
+			PRAGMA_ENABLE_DEPRECATION_WARNINGS
 		}
 	}
 
diff --git a/Engine/Source/Editor/BlueprintGraph/Private/K2Node_ActorBoundEvent.cpp b/Engine/Source/Editor/BlueprintGraph/Private/K2Node_ActorBoundEvent.cpp
--- a/Engine/Source/Editor/BlueprintGraph/Private/K2Node_ActorBoundEvent.cpp
+++ b/Engine/Source/Editor/BlueprintGraph/Private/K2Node_ActorBoundEvent.cpp
@@ -192,7 +192,7 @@ UMulticastDelegateProperty* UK2Node_ActorBoundEvent::GetTargetDelegateProperty()
 	// If we couldn't find the target delegate, then try to find it in the property remap table
 	if (!TargetDelegateProp)
 	{
-		UMulticastDelegateProperty* NewProperty = Cast<UMulticastDelegateProperty>(FMemberReference::FindRemappedField(DelegateOwnerClass, DelegatePropertyName));
+		UMulticastDelegateProperty* NewProperty = FMemberReference::FindRemappedField<UMulticastDelegateProperty>(DelegateOwnerClass, DelegatePropertyName);
 		if (NewProperty)
 		{
 			// Found a remapped property, update the node
diff --git a/Engine/Source/Editor/BlueprintGraph/Private/K2Node_BreakStruct.cpp b/Engine/Source/Editor/BlueprintGraph/Private/K2Node_BreakStruct.cpp
--- a/Engine/Source/Editor/BlueprintGraph/Private/K2Node_BreakStruct.cpp
+++ b/Engine/Source/Editor/BlueprintGraph/Private/K2Node_BreakStruct.cpp
@@ -317,7 +317,7 @@ UK2Node::ERedirectType UK2Node_BreakStruct::DoPinsMatchForReconstruction(const U
 	ERedirectType Result = UK2Node::DoPinsMatchForReconstruction(NewPin, NewPinIndex, OldPin, OldPinIndex);
 	if ((ERedirectType_None == Result) && DoRenamedPinsMatch(NewPin, OldPin, true))
 	{
-		Result = ERedirectType_Custom;
+		Result = ERedirectType_Name;
 	}
 	else if ((ERedirectType_None == Result) && NewPin && OldPin)
 	{
@@ -326,19 +326,16 @@ UK2Node::ERedirectType UK2Node_BreakStruct::DoPinsMatchForReconstruction(const U
 			const UEdGraphSchema_K2* K2Schema = GetDefault<UEdGraphSchema_K2>();
 			if (K2Schema->ArePinTypesCompatible( NewPin->PinType, OldPin->PinType))
 			{
-				Result = ERedirectType_Custom;
+				Result = ERedirectType_Name;
 			}
 		}
 		else if ((EGPD_Output == NewPin->Direction) && (EGPD_Output == OldPin->Direction))
 		{
-			TMap<FName, FName>* StructRedirects = UStruct::TaggedPropertyRedirects.Find(StructType->GetFName());
-			if (StructRedirects)
+			FName RedirectedPinName = UProperty::FindRedirectedPropertyName(StructType, FName(*OldPin->PinName));
+
+			if (RedirectedPinName != NAME_Name)
 			{
-				FName* PropertyRedirect = StructRedirects->Find(FName(*OldPin->PinName));
-				if (PropertyRedirect)
-				{
-					Result = ((FCString::Stricmp(*PropertyRedirect->ToString(), *NewPin->PinName) != 0) ? ERedirectType_None : ERedirectType_Name);
-				}
+				Result = ((FCString::Stricmp(*RedirectedPinName.ToString(), *NewPin->PinName) != 0) ? ERedirectType_None : ERedirectType_Name);
 			}
 		}
 	}
diff --git a/Engine/Source/Editor/BlueprintGraph/Private/K2Node_CallFunction.cpp b/Engine/Source/Editor/BlueprintGraph/Private/K2Node_CallFunction.cpp
--- a/Engine/Source/Editor/BlueprintGraph/Private/K2Node_CallFunction.cpp
+++ b/Engine/Source/Editor/BlueprintGraph/Private/K2Node_CallFunction.cpp
@@ -583,7 +583,7 @@ void UK2Node_CallFunction::AllocateDefaultPins()
 
 		if (ParentClass != NULL)
 		{
-			if (UFunction* NewFunction = Cast<UFunction>(FMemberReference::FindRemappedField(ParentClass, FunctionReference.GetMemberName())))
+			if (UFunction* NewFunction = FMemberReference::FindRemappedField<UFunction>(ParentClass, FunctionReference.GetMemberName()))
 			{
 				// Found a remapped property, update the node
 				Function = NewFunction;
diff --git a/Engine/Source/Editor/BlueprintGraph/Private/K2Node_CreateDelegate.cpp b/Engine/Source/Editor/BlueprintGraph/Private/K2Node_CreateDelegate.cpp
--- a/Engine/Source/Editor/BlueprintGraph/Private/K2Node_CreateDelegate.cpp
+++ b/Engine/Source/Editor/BlueprintGraph/Private/K2Node_CreateDelegate.cpp
@@ -17,8 +17,10 @@
 struct FK2Node_CreateDelegate_Helper
 {
 	static FString DelegateOutputName;
+	static FString InputObjectName; // Deprecated, for fixup
 };
 FString FK2Node_CreateDelegate_Helper::DelegateOutputName(TEXT("OutputDelegate"));
+FString FK2Node_CreateDelegate_Helper::InputObjectName(TEXT("InputObject"));
 
 UK2Node_CreateDelegate::UK2Node_CreateDelegate(const FObjectInitializer& ObjectInitializer)
 	: Super(ObjectInitializer)
@@ -42,6 +44,19 @@ void UK2Node_CreateDelegate::AllocateDefaultPins()
 	Super::AllocateDefaultPins();
 }
 
+UK2Node::ERedirectType UK2Node_CreateDelegate::DoPinsMatchForReconstruction(const UEdGraphPin* NewPin, int32 NewPinIndex, const UEdGraphPin* OldPin, int32 OldPinIndex) const
+{
+	const UEdGraphSchema_K2* K2Schema = GetDefault<UEdGraphSchema_K2>();
+
+	// Handles remap of InputObject to Self, from 4.10 time frame
+	if (OldPin->PinName == FK2Node_CreateDelegate_Helper::InputObjectName && NewPin->PinName == K2Schema->PN_Self)
+	{
+		return ERedirectType_Name;
+	}
+
+	return Super::DoPinsMatchForReconstruction(NewPin, NewPinIndex, OldPin, OldPinIndex);
+}
+
 bool UK2Node_CreateDelegate::IsValid(FString* OutMsg, bool bDontUseSkeletalClassForSelf) const
 {
 	FName FunctionName = GetFunctionName();
diff --git a/Engine/Source/Editor/BlueprintGraph/Private/K2Node_DelegateSet.cpp b/Engine/Source/Editor/BlueprintGraph/Private/K2Node_DelegateSet.cpp
--- a/Engine/Source/Editor/BlueprintGraph/Private/K2Node_DelegateSet.cpp
+++ b/Engine/Source/Editor/BlueprintGraph/Private/K2Node_DelegateSet.cpp
@@ -200,7 +200,7 @@ UFunction* UK2Node_DelegateSet::GetDelegateSignature()
 	if( !DelegateProperty )
 	{
 		// Attempt to find a remapped delegate property
-		UMulticastDelegateProperty* NewProperty = Cast<UMulticastDelegateProperty>(FMemberReference::FindRemappedField(DelegatePropertyClass, DelegatePropertyName));
+		UMulticastDelegateProperty* NewProperty = FMemberReference::FindRemappedField<UMulticastDelegateProperty>(DelegatePropertyClass, DelegatePropertyName);
 		if( NewProperty )
 		{
 			// Found a remapped property, update the node
@@ -220,7 +220,7 @@ UFunction* UK2Node_DelegateSet::GetDelegateSignature() const
 	if( !DelegateProperty )
 	{
 		// Attempt to find a remapped delegate property
-		DelegateProperty = Cast<UMulticastDelegateProperty>(FMemberReference::FindRemappedField(DelegatePropertyClass, DelegatePropertyName));
+		DelegateProperty = FMemberReference::FindRemappedField<UMulticastDelegateProperty>(DelegatePropertyClass, DelegatePropertyName);
 	}
 
 	return (DelegateProperty != NULL) ? DelegateProperty->SignatureFunction : NULL;
diff --git a/Engine/Source/Editor/BlueprintGraph/Private/K2Node_MakeStruct.cpp b/Engine/Source/Editor/BlueprintGraph/Private/K2Node_MakeStruct.cpp
--- a/Engine/Source/Editor/BlueprintGraph/Private/K2Node_MakeStruct.cpp
+++ b/Engine/Source/Editor/BlueprintGraph/Private/K2Node_MakeStruct.cpp
@@ -302,7 +302,7 @@ UK2Node::ERedirectType UK2Node_MakeStruct::DoPinsMatchForReconstruction(const UE
 	ERedirectType Result = UK2Node::DoPinsMatchForReconstruction(NewPin, NewPinIndex, OldPin, OldPinIndex);
 	if ((ERedirectType_None == Result) && DoRenamedPinsMatch(NewPin, OldPin, false))
 	{
-		Result = ERedirectType_Custom;
+		Result = ERedirectType_Name;
 	}
 	else if ((ERedirectType_None == Result) && NewPin && OldPin)
 	{
@@ -311,19 +311,16 @@ UK2Node::ERedirectType UK2Node_MakeStruct::DoPinsMatchForReconstruction(const UE
 			const UEdGraphSchema_K2* K2Schema = GetDefault<UEdGraphSchema_K2>();
 			if (K2Schema->ArePinTypesCompatible(NewPin->PinType, OldPin->PinType))
 			{
-				Result = ERedirectType_Custom;
+				Result = ERedirectType_Name;
 			}
 		}
 		else if ((EGPD_Input == NewPin->Direction) && (EGPD_Input == OldPin->Direction))
 		{
-			TMap<FName, FName>* StructRedirects = UStruct::TaggedPropertyRedirects.Find(StructType->GetFName());
-			if (StructRedirects)
+			FName RedirectedPinName = UProperty::FindRedirectedPropertyName(StructType, FName(*OldPin->PinName));
+
+			if (RedirectedPinName != NAME_Name)
 			{
-				FName* PropertyRedirect = StructRedirects->Find(FName(*OldPin->PinName));
-				if (PropertyRedirect)
-				{
-					Result = ((FCString::Stricmp(*PropertyRedirect->ToString(), *NewPin->PinName) != 0) ? ERedirectType_None : ERedirectType_Name);
-				}
+				Result = ((FCString::Stricmp(*RedirectedPinName.ToString(), *NewPin->PinName) != 0) ? ERedirectType_None : ERedirectType_Name);
 			}
 		}
 	}
diff --git a/Engine/Source/Editor/BlueprintGraph/Private/K2Node_Message.cpp b/Engine/Source/Editor/BlueprintGraph/Private/K2Node_Message.cpp
--- a/Engine/Source/Editor/BlueprintGraph/Private/K2Node_Message.cpp
+++ b/Engine/Source/Editor/BlueprintGraph/Private/K2Node_Message.cpp
@@ -168,8 +168,8 @@ void UK2Node_Message::ExpandNode(class FKismetCompilerContext& CompilerContext,
 		UFunction* MessageNodeFunction = GetTargetFunction();
 		if (MessageNodeFunction == NULL)
 		{
-			//@TODO: Why do this here in teh compiler, it's already done on AllocateDefaultPins() during on-load node reconstruction
-			MessageNodeFunction = Cast<UFunction>(FMemberReference::FindRemappedField(FunctionReference.GetMemberParentClass(GetBlueprintClassFromNode()), FunctionReference.GetMemberName()));
+			//@TODO: Why do this here in the compiler, it's already done on AllocateDefaultPins() during on-load node reconstruction
+			MessageNodeFunction = FMemberReference::FindRemappedField<UFunction>(FunctionReference.GetMemberParentClass(GetBlueprintClassFromNode()), FunctionReference.GetMemberName());
 		}
 
 		if (MessageNodeFunction == NULL)
diff --git a/Engine/Source/Editor/BlueprintGraph/Private/K2Node_SpawnActorFromClass.cpp b/Engine/Source/Editor/BlueprintGraph/Private/K2Node_SpawnActorFromClass.cpp
--- a/Engine/Source/Editor/BlueprintGraph/Private/K2Node_SpawnActorFromClass.cpp
+++ b/Engine/Source/Editor/BlueprintGraph/Private/K2Node_SpawnActorFromClass.cpp
@@ -22,6 +22,7 @@ struct FK2Node_SpawnActorFromClassHelper
 	static FString WorldContextPinName;
 	static FString ClassPinName;
 	static FString SpawnTransformPinName;
+	static FString SpawnEvenIfCollidingPinName;
 	static FString NoCollisionFailPinName;
 	static FString CollisionHandlingOverridePinName;
 	static FString OwnerPinName;
@@ -30,7 +31,8 @@ struct FK2Node_SpawnActorFromClassHelper
 FString FK2Node_SpawnActorFromClassHelper::WorldContextPinName(TEXT("WorldContextObject"));
 FString FK2Node_SpawnActorFromClassHelper::ClassPinName(TEXT("Class"));
 FString FK2Node_SpawnActorFromClassHelper::SpawnTransformPinName(TEXT("SpawnTransform"));
-FString FK2Node_SpawnActorFromClassHelper::NoCollisionFailPinName(TEXT("SpawnEvenIfColliding"));		// deprecated pin, name kept for backwards compat
+FString FK2Node_SpawnActorFromClassHelper::SpawnEvenIfCollidingPinName(TEXT("SpawnEvenIfColliding"));		// deprecated pin, name kept for backwards compat
+FString FK2Node_SpawnActorFromClassHelper::NoCollisionFailPinName(TEXT("bNoCollisionFail"));		// deprecated pin, name kept for backwards compat
 FString FK2Node_SpawnActorFromClassHelper::CollisionHandlingOverridePinName(TEXT("CollisionHandlingOverride"));
 FString FK2Node_SpawnActorFromClassHelper::OwnerPinName(TEXT("Owner"));
 
@@ -153,7 +155,7 @@ void UK2Node_SpawnActorFromClass::MaybeUpdateCollisionPin(TArray<UEdGraphPin*>&
 	// see if there's a bNoCollisionFail pin
 	for (UEdGraphPin* Pin : OldPins)
 	{
-		if (Pin->PinName == FK2Node_SpawnActorFromClassHelper::NoCollisionFailPinName)
+		if (Pin->PinName == FK2Node_SpawnActorFromClassHelper::NoCollisionFailPinName || Pin->PinName == FK2Node_SpawnActorFromClassHelper::SpawnEvenIfCollidingPinName)
 		{
 			bool bHadOldCollisionPin = true;
 			if (Pin->LinkedTo.Num() == 0)
diff --git a/Engine/Source/Editor/BlueprintGraph/Private/K2Node_Variable.cpp b/Engine/Source/Editor/BlueprintGraph/Private/K2Node_Variable.cpp
--- a/Engine/Source/Editor/BlueprintGraph/Private/K2Node_Variable.cpp
+++ b/Engine/Source/Editor/BlueprintGraph/Private/K2Node_Variable.cpp
@@ -444,13 +444,13 @@ void UK2Node_Variable::ValidateNodeDuringCompilation(class FCompilerResultsLog&
 			{
 				FString const VarName = VariableReference.GetMemberName().ToString();
 
-				FText const WarningFormat = LOCTEXT("VariableNotFound", "Could not find a variable named \"%s\" in '%s'.\nMake sure '%s' has been compiled for @@");
-				MessageLog.Warning(*FString::Printf(*WarningFormat.ToString(), *VarName, *OwnerName, *OwnerName), this);
+				FText const WarningFormat = LOCTEXT("VariableNotFoundFmt", "Could not find a variable named \"{0}\" in '{1}'.\nMake sure '{2}' has been compiled for @@");
+				MessageLog.Warning(*FText::Format(WarningFormat, FText::FromString(VarName), FText::FromString(OwnerName), FText::FromString(OwnerName)).ToString(), this);
 			}
 		}
 		else
 		{
-			MessageLog.Warning(*FString::Printf(*LOCTEXT("VariableDeprecated", "Variable '%s' for @@ was deprecated.  Please update it.").ToString(), *VariableReference.GetMemberName().ToString()), this);
+			MessageLog.Warning(*FText::Format(LOCTEXT("VariableDeprecatedFmt", "Variable '{0}' for @@ was deprecated.  Please update it."), FText::FromString(VariableReference.GetMemberName().ToString())).ToString(), this);
 		}
 	}
 
@@ -573,27 +573,24 @@ void UK2Node_Variable::ReconstructNode()
 	{
 		bool bRemappedProperty = false;
 		UClass* SearchClass = VarClass;
-		while (SearchClass != NULL)
+		while (SearchClass != nullptr)
 		{
-			const TMap<FName, FName>* const ClassTaggedPropertyRedirects = UStruct::TaggedPropertyRedirects.Find( SearchClass->GetFName() );
-			if (ClassTaggedPropertyRedirects)
+			FName NewPropertyName = UProperty::FindRedirectedPropertyName(SearchClass, VariableReference.GetMemberName());
+
+			if (NewPropertyName != NAME_None)
 			{
-				const FName* const NewPropertyName = ClassTaggedPropertyRedirects->Find( VariableReference.GetMemberName() );
-				if (NewPropertyName)
+				if (VariableReference.IsSelfContext())
 				{
-					if (VariableReference.IsSelfContext())
-					{
-						VariableReference.SetSelfMember( *NewPropertyName );
-					}
-					else
-					{
-						VariableReference.SetExternalMember( *NewPropertyName, VarClass );
-					}
-
-					// found, can break
-					bRemappedProperty = true;
-					break;
+					VariableReference.SetSelfMember(NewPropertyName);
+				}
+				else
+				{
+					VariableReference.SetExternalMember(NewPropertyName, VarClass);
 				}
+
+				// found, can break
+				bRemappedProperty = true;
+				break;
 			}
 
 			SearchClass = SearchClass->GetSuperClass();
diff --git a/Engine/Source/Editor/EditorStyle/Private/SlateEditorStyle.cpp b/Engine/Source/Editor/EditorStyle/Private/SlateEditorStyle.cpp
--- a/Engine/Source/Editor/EditorStyle/Private/SlateEditorStyle.cpp
+++ b/Engine/Source/Editor/EditorStyle/Private/SlateEditorStyle.cpp
@@ -5895,13 +5895,11 @@ void FSlateEditorStyle::FStyle::SetupClassIconsAndThumbnails()
 			TEXT("CameraRig_Rail")
 		};
 
-		const TCHAR* Path = TEXT("Icons/AssetIcons/%s_%dx");
-		
 		for (int32 TypeIndex = 0; TypeIndex < ARRAY_COUNT(AssetTypes); ++TypeIndex)
 		{
 			const TCHAR* Type = AssetTypes[TypeIndex];
-			Set( *FString::Printf(TEXT("ClassIcon.%s"), Type),		new IMAGE_BRUSH(FString::Printf(Path, Type, 16), Icon16x16 ) );
-			Set( *FString::Printf(TEXT("ClassThumbnail.%s"), Type),	new IMAGE_BRUSH(FString::Printf(Path, Type, 64), Icon64x64 ) );
+			Set( *FString::Printf(TEXT("ClassIcon.%s"), Type),		new IMAGE_BRUSH(FString::Printf(TEXT("Icons/AssetIcons/%s_%dx"), Type, 16), Icon16x16 ) );
+			Set( *FString::Printf(TEXT("ClassThumbnail.%s"), Type),	new IMAGE_BRUSH(FString::Printf(TEXT("Icons/AssetIcons/%s_%dx"), Type, 64), Icon64x64 ) );
 		}
 	}
 #endif
diff --git a/Engine/Source/Editor/GameProjectGeneration/Private/GameProjectUtils.cpp b/Engine/Source/Editor/GameProjectGeneration/Private/GameProjectUtils.cpp
--- a/Engine/Source/Editor/GameProjectGeneration/Private/GameProjectUtils.cpp
+++ b/Engine/Source/Editor/GameProjectGeneration/Private/GameProjectUtils.cpp
@@ -1248,7 +1248,7 @@ bool GameProjectUtils::GenerateProjectFromScratch(const FProjectInformation& InP
 	return true;
 }
 
-bool GameProjectUtils::CreateProjectFromTemplate(const FProjectInformation& InProjectInfo, FText& OutFailReason, FText& OutFailLog, TArray<FString>* OutCreatedFiles)
+bool GameProjectUtils::CreateProjectFromTemplate(const FProjectInformation& InProjectInfo, FText& OutFailReason, FText& OutFailLog,TArray<FString>* OutCreatedFiles)
 {
 	FScopedSlowTask SlowTask(10);
 
@@ -1284,9 +1284,9 @@ bool GameProjectUtils::CreateProjectFromTemplate(const FProjectInformation& InPr
 
 	// Form a list of all extensions we care about
 	TSet<FString> ReplacementsInFilesExtensions;
-	for ( auto ReplacementIt = TemplateDefs->ReplacementsInFiles.CreateConstIterator(); ReplacementIt; ++ReplacementIt )
+	for ( const FTemplateReplacement& Replacement : TemplateDefs->ReplacementsInFiles )
 	{
-		ReplacementsInFilesExtensions.Append((*ReplacementIt).Extensions);
+		ReplacementsInFilesExtensions.Append(Replacement.Extensions);
 	}
 
 	// Keep a list of created files so we can delete them if project creation fails
@@ -1304,10 +1304,8 @@ bool GameProjectUtils::CreateProjectFromTemplate(const FProjectInformation& InPr
 	{
 		// Open a new feedback scope for the loop so we can report how far through the copy we are
 		FScopedSlowTask InnerSlowTask(FilesToCopy.Num());
-		for ( auto FileIt = FilesToCopy.CreateConstIterator(); FileIt; ++FileIt )
+		for ( const FString& SrcFilename : FilesToCopy )
 		{
-			const FString SrcFilename = (*FileIt);
-
 			// Update the progress
 			FFormatNamedArguments Args;
 			Args.Add( TEXT("SrcFilename"), FText::FromString( FPaths::GetCleanFilename(SrcFilename) ) );
@@ -1317,47 +1315,24 @@ bool GameProjectUtils::CreateProjectFromTemplate(const FProjectInformation& InPr
 			const FString SrcFileSubpath = SrcFilename.RightChop(SrcFolder.Len() + 1);
 
 			// Skip any files that were configured to be ignored
-			bool bThisFileIsIgnored = false;
-			for ( auto IgnoreIt = TemplateDefs->FilesToIgnore.CreateConstIterator(); IgnoreIt; ++IgnoreIt )
-			{
-				if ( SrcFileSubpath == *IgnoreIt )
-				{
-					// This file was marked as "ignored"
-					bThisFileIsIgnored = true;
-					break;
-				}
-			}
-
-			if ( bThisFileIsIgnored )
+			if ( TemplateDefs->FilesToIgnore.Contains(SrcFileSubpath) )
 			{
 				// This file was marked as "ignored"
 				continue;
 			}
 
 			// Skip any folders that were configured to be ignored
-			bool bThisFolderIsIgnored = false;
-			for ( auto IgnoreIt = TemplateDefs->FoldersToIgnore.CreateConstIterator(); IgnoreIt; ++IgnoreIt )
-			{
-				if ( SrcFileSubpath.StartsWith((*IgnoreIt) + TEXT("/") ) )
-				{
-					// This folder was marked as "ignored"
-					UE_LOG(LogGameProjectGeneration, Verbose, TEXT("'%s': Skipping as it is in an ignored folder '%s'"), *SrcFilename, **IgnoreIt);
-					bThisFolderIsIgnored = true;
-					break;
-				}
-			}
-
-			if ( bThisFolderIsIgnored )
+			if ( const FString* IgnoredFolder = TemplateDefs->FoldersToIgnore.FindByPredicate([&SrcFileSubpath](const FString& Ignore){ return SrcFileSubpath.StartsWith(Ignore + TEXT("/")); }) )
 			{
 				// This folder was marked as "ignored"
+				UE_LOG(LogGameProjectGeneration, Verbose, TEXT("'%s': Skipping as it is in an ignored folder '%s'"), *SrcFilename, **IgnoredFolder);
 				continue;
 			}
 
 			// Retarget any folders that were chosen to be renamed by choosing a new destination subpath now
 			FString DestFileSubpathWithoutFilename = FPaths::GetPath(SrcFileSubpath) + TEXT("/");
-			for ( auto RenameIt = TemplateDefs->FolderRenames.CreateConstIterator(); RenameIt; ++RenameIt )
+			for ( const FTemplateFolderRename& FolderRename : TemplateDefs->FolderRenames )
 			{
-				const FTemplateFolderRename& FolderRename = *RenameIt;
 				if ( SrcFileSubpath.StartsWith(FolderRename.From + TEXT("/")) )
 				{
 					// This was a file in a renamed folder. Retarget to the new location
@@ -1369,9 +1344,8 @@ bool GameProjectUtils::CreateProjectFromTemplate(const FProjectInformation& InPr
 			// Retarget any files that were chosen to have parts of their names replaced here
 			FString DestBaseFilename = FPaths::GetBaseFilename(SrcFileSubpath);
 			const FString FileExtension = FPaths::GetExtension(SrcFileSubpath);
-			for ( auto ReplacementIt = TemplateDefs->FilenameReplacements.CreateConstIterator(); ReplacementIt; ++ReplacementIt )
+			for ( const FTemplateReplacement& Replacement : TemplateDefs->FilenameReplacements )
 			{
-				const FTemplateReplacement& Replacement = *ReplacementIt;
 				if ( Replacement.Extensions.Contains( FileExtension ) )
 				{
 					// This file matched a filename replacement extension, apply it now
@@ -1421,19 +1395,17 @@ bool GameProjectUtils::CreateProjectFromTemplate(const FProjectInformation& InPr
 		FScopedSlowTask InnerSlowTask(FilesThatNeedContentsReplaced.Num());
 
 		// Open all files with the specified extensions and replace text
-		for ( auto FileIt = FilesThatNeedContentsReplaced.CreateConstIterator(); FileIt; ++FileIt )
+		for ( const FString& FileToFix : FilesThatNeedContentsReplaced )
 		{
 			InnerSlowTask.EnterProgressFrame();
 
-			const FString FileToFix = *FileIt;
 			bool bSuccessfullyProcessed = false;
 
 			FString FileContents;
 			if ( FFileHelper::LoadFileToString(FileContents, *FileToFix) )
 			{
-				for ( auto ReplacementIt = TemplateDefs->ReplacementsInFiles.CreateConstIterator(); ReplacementIt; ++ReplacementIt )
+				for ( const FTemplateReplacement& Replacement : TemplateDefs->ReplacementsInFiles )
 				{
-					const FTemplateReplacement& Replacement = *ReplacementIt;
 					if ( Replacement.Extensions.Contains( FPaths::GetExtension(FileToFix) ) )
 					{
 						FileContents = FileContents.Replace(*Replacement.From, *Replacement.To, Replacement.bCaseSensitive ? ESearchCase::CaseSensitive : ESearchCase::IgnoreCase);
@@ -1480,19 +1452,18 @@ bool GameProjectUtils::CreateProjectFromTemplate(const FProjectInformation& InPr
 	// Fixup specific ini values
 	TArray<FTemplateConfigValue> ConfigValuesToSet;
 	TemplateDefs->AddConfigValues(ConfigValuesToSet, TemplateName, ProjectName, InProjectInfo.bShouldGenerateCode);
-	new (ConfigValuesToSet) FTemplateConfigValue(TEXT("DefaultGame.ini"), TEXT("/Script/EngineSettings.GeneralProjectSettings"), TEXT("ProjectID"), FGuid::NewGuid().ToString(), /*InShouldReplaceExistingValue=*/true);
+	ConfigValuesToSet.Emplace(TEXT("DefaultGame.ini"), TEXT("/Script/EngineSettings.GeneralProjectSettings"), TEXT("ProjectID"), FGuid::NewGuid().ToString(), /*InShouldReplaceExistingValue=*/true);
 
 	// Add all classname fixups
-	for ( auto RenameIt = ClassRenames.CreateConstIterator(); RenameIt; ++RenameIt )
+	for ( const TPair<FString, FString>& Rename : ClassRenames )
 	{
-		const FString ClassRedirectString = FString::Printf(TEXT("(OldClassName=\"%s\",NewClassName=\"%s\")"), *RenameIt.Key(), *RenameIt.Value());
-		new (ConfigValuesToSet) FTemplateConfigValue(TEXT("DefaultEngine.ini"), TEXT("/Script/Engine.Engine"), TEXT("+ActiveClassRedirects"), *ClassRedirectString, /*InShouldReplaceExistingValue=*/false);
+		const FString ClassRedirectString = FString::Printf(TEXT("(OldClassName=\"%s\",NewClassName=\"%s\")"), *Rename.Key, *Rename.Value);
+		ConfigValuesToSet.Emplace(TEXT("DefaultEngine.ini"), TEXT("/Script/Engine.Engine"), TEXT("+ActiveClassRedirects"), *ClassRedirectString, /*InShouldReplaceExistingValue=*/false);
 	}
 
 	// Fix all specified config values
-	for ( auto ConfigIt = ConfigValuesToSet.CreateConstIterator(); ConfigIt; ++ConfigIt )
+	for ( const FTemplateConfigValue& ConfigValue : ConfigValuesToSet )
 	{
-		const FTemplateConfigValue& ConfigValue = *ConfigIt;
 		const FString IniFilename = ProjectConfigPath / ConfigValue.ConfigFile;
 		bool bSuccessfullyProcessed = false;
 
@@ -1503,9 +1474,9 @@ bool GameProjectUtils::CreateProjectFromTemplate(const FProjectInformation& InPr
 			const FString TargetSection = ConfigValue.ConfigSection;
 			FString CurSection;
 			bool bFoundTargetKey = false;
-			for ( auto LineIt = FileLines.CreateConstIterator(); LineIt; ++LineIt )
+			for ( const FString& LineIn : FileLines )
 			{
-				FString Line = *LineIt;
+				FString Line = LineIn;
 				Line.Trim().TrimTrailing();
 
 				bool bShouldExcludeLineFromOutput = false;
@@ -1560,7 +1531,7 @@ bool GameProjectUtils::CreateProjectFromTemplate(const FProjectInformation& InPr
 				if ( !bShouldExcludeLineFromOutput )
 				{
 					FileOutput += Line;
-					if ( LineIt.GetIndex() < FileLines.Num() - 1 )
+					if ( &LineIn != &FileLines.Last() )
 					{
 						// Add a line terminator on every line except the last
 						FileOutput += LINE_TERMINATOR;
@@ -1597,16 +1568,16 @@ bool GameProjectUtils::CreateProjectFromTemplate(const FProjectInformation& InPr
 	}
 
 	// Insert any required feature packs (EG starter content) into ini file. These will be imported automatically when the editor is first run
-	if ( !InsertFeaturePacksIntoINIFile(InProjectInfo, OutFailReason) )
+	if(!InsertFeaturePacksIntoINIFile(InProjectInfo, OutFailReason))
 	{
 		return false;
 	}
 
-	if ( !AddSharedContentToProject(InProjectInfo, CreatedFiles, OutFailReason) )
+	if( !AddSharedContentToProject(InProjectInfo, CreatedFiles, OutFailReason ) )
 	{
 		return false;
 	}
-
+	
 	
 	SlowTask.EnterProgressFrame();
 
@@ -1614,7 +1585,7 @@ bool GameProjectUtils::CreateProjectFromTemplate(const FProjectInformation& InPr
 	{
 		// Load the source project
 		FProjectDescriptor Project;
-		if ( !Project.Load(InProjectInfo.TemplateFile, OutFailReason) )
+		if(!Project.Load(InProjectInfo.TemplateFile, OutFailReason))
 		{
 			return false;
 		}
@@ -1626,9 +1597,8 @@ bool GameProjectUtils::CreateProjectFromTemplate(const FProjectInformation& InPr
 		// Fix up module names
 		const FString BaseSourceName = FPaths::GetBaseFilename(InProjectInfo.TemplateFile);
 		const FString BaseNewName = FPaths::GetBaseFilename(InProjectInfo.ProjectFilename);
-		for ( auto ModuleIt = Project.Modules.CreateIterator(); ModuleIt; ++ModuleIt )
+		for ( FModuleDescriptor& ModuleInfo : Project.Modules )
 		{
-			FModuleDescriptor& ModuleInfo = *ModuleIt;
 			ModuleInfo.Name = FName(*ModuleInfo.Name.ToString().Replace(*BaseSourceName, *BaseNewName));
 		}
 
diff --git a/Engine/Source/Editor/Persona/Private/AnimationEditorViewportClient.cpp b/Engine/Source/Editor/Persona/Private/AnimationEditorViewportClient.cpp
--- a/Engine/Source/Editor/Persona/Private/AnimationEditorViewportClient.cpp
+++ b/Engine/Source/Editor/Persona/Private/AnimationEditorViewportClient.cpp
@@ -634,7 +634,7 @@ void FAnimationViewportClient::DisplayInfo(FCanvas* Canvas, FSceneView* View, bo
 
 		if (MaterialsThatNeedMorphFlagOn.Num() > 0)
 		{
-			InfoString = FString::Printf( *LOCTEXT("MorphSupportNeeded", "The following materials need morph support ('Used with Morph Targets' in material editor):").ToString() );
+			InfoString = LOCTEXT("MorphSupportNeeded", "The following materials need morph support ('Used with Morph Targets' in material editor):").ToString();
 			Canvas->DrawShadowedString( CurXOffset, CurYOffset, *InfoString, GEngine->GetSmallFont(), HeadlineColour );
 
 			CurYOffset += YL + 2;
@@ -651,7 +651,7 @@ void FAnimationViewportClient::DisplayInfo(FCanvas* Canvas, FSceneView* View, bo
 
 		if (MaterialsThatNeedSaving.Num() > 0)
 		{
-			InfoString = FString::Printf( *LOCTEXT("MaterialsNeedSaving", "The following materials need saving to fully support morph targets:").ToString() );
+			InfoString = LOCTEXT("MaterialsNeedSaving", "The following materials need saving to fully support morph targets:").ToString();
 			Canvas->DrawShadowedString( CurXOffset, CurYOffset, *InfoString, GEngine->GetSmallFont(), HeadlineColour );
 
 			CurYOffset += YL + 2;
@@ -876,7 +876,7 @@ void FAnimationViewportClient::DisplayInfo(FCanvas* Canvas, FSceneView* View, bo
 	{
 		// Notify the user if they are isolating a mesh section.
 		CurYOffset += YL + 2;
-		InfoString = FString::Printf(*LOCTEXT("MeshSectionsHiddenWarning", "Mesh Sections Hidden").ToString());
+		InfoString = LOCTEXT("MeshSectionsHiddenWarning", "Mesh Sections Hidden").ToString();
 		Canvas->DrawShadowedString(CurXOffset, CurYOffset, *InfoString, GEngine->GetSmallFont(), SubHeadlineColour);
 		
 	}
diff --git a/Engine/Source/Editor/UnrealEd/Private/Factories/ActorFactory.cpp b/Engine/Source/Editor/UnrealEd/Private/Factories/ActorFactory.cpp
--- a/Engine/Source/Editor/UnrealEd/Private/Factories/ActorFactory.cpp
+++ b/Engine/Source/Editor/UnrealEd/Private/Factories/ActorFactory.cpp
@@ -585,7 +585,7 @@ UActorFactoryTextRender::UActorFactoryTextRender(const FObjectInitializer& Objec
 {
 	// Property initialization
 	DisplayName = LOCTEXT("TextRenderDisplayName", "Text Render");
-	NewActorClass = ATextRenderActor::GetPrivateStaticClass(L"...");
+	NewActorClass = ATextRenderActor::StaticClass();
 	bUseSurfaceOrientation = true;
 }
 
@@ -597,7 +597,7 @@ UActorFactorySubDSurface::UActorFactorySubDSurface(const FObjectInitializer& Obj
 {
 	// Property initialization
 	DisplayName = LOCTEXT("SubDSurfaceDisplayName", "Subdivision Surface");
-	NewActorClass = ASubDSurfaceActor::GetPrivateStaticClass(L"...");
+	NewActorClass = ASubDSurfaceActor::StaticClass();
 	bUseSurfaceOrientation = true;
 }
 
diff --git a/Engine/Source/Editor/UnrealEd/Private/FileHelpers.cpp b/Engine/Source/Editor/UnrealEd/Private/FileHelpers.cpp
--- a/Engine/Source/Editor/UnrealEd/Private/FileHelpers.cpp
+++ b/Engine/Source/Editor/UnrealEd/Private/FileHelpers.cpp
@@ -74,7 +74,7 @@ TSet<FString> FEditorFileUtils::PackagesNotSavedDuringSaveAll;
 TSet<FString> FEditorFileUtils::PackagesNotToPromptAnyMore;
 
 static const FString InvalidFilenames[] = {
-	TEXT("CON"), TEXT("PRN"), TEXT("AUX"), TEXT("CLOCK$"), TEXT("NUL"),
+	TEXT("CON"), TEXT("PRN"), TEXT("AUX"), TEXT("CLOCK$"), TEXT("NUL"), TEXT("NONE"),
 	TEXT("COM1"), TEXT("COM2"), TEXT("COM3"), TEXT("COM4"), TEXT("COM5"), TEXT("COM6"), TEXT("COM7"), TEXT("COM8"), TEXT("COM9"),
 	TEXT("LPT1"), TEXT("LPT2"), TEXT("LPT3"), TEXT("LPT4"), TEXT("LPT5"), TEXT("LPT6"), TEXT("LPT7"), TEXT("LPT8"), TEXT("LPT9")
 };
diff --git a/Engine/Source/Editor/UnrealEd/Private/Kismet2/BlueprintEditorUtils.cpp b/Engine/Source/Editor/UnrealEd/Private/Kismet2/BlueprintEditorUtils.cpp
--- a/Engine/Source/Editor/UnrealEd/Private/Kismet2/BlueprintEditorUtils.cpp
+++ b/Engine/Source/Editor/UnrealEd/Private/Kismet2/BlueprintEditorUtils.cpp
@@ -983,6 +983,7 @@ struct FRegenerationHelper
 		FLinkerLoad* Linker = Obj->GetLinker();
 		if (Linker && !Obj->HasAnyFlags(RF_LoadCompleted))
 		{
+			check(!GEventDrivenLoaderEnabled);
 			Obj->SetFlags(RF_NeedLoad);
 			Linker->Preload(Obj);
 			return true;
diff --git a/Engine/Source/Editor/UnrealEd/Private/Kismet2/KismetReinstanceUtilities.cpp b/Engine/Source/Editor/UnrealEd/Private/Kismet2/KismetReinstanceUtilities.cpp
--- a/Engine/Source/Editor/UnrealEd/Private/Kismet2/KismetReinstanceUtilities.cpp
+++ b/Engine/Source/Editor/UnrealEd/Private/Kismet2/KismetReinstanceUtilities.cpp
@@ -218,7 +218,6 @@ FBlueprintCompileReinstancer::FBlueprintCompileReinstancer(UClass* InClassToRein
 	, DuplicatedClass(nullptr)
 	, OriginalCDO(nullptr)
 	, bHasReinstanced(false)
-	, bSkipGarbageCollection(bSkipGC)
 	, ReinstClassType(RCT_Unknown)
 	, ClassToReinstanceDefaultValuesCRC(0)
 	, bIsRootReinstancer(false)
@@ -636,8 +635,9 @@ void FBlueprintCompileReinstancer::CompileChildren()
 			// skeletons were kept in-sync (updated/reinstanced when the parent 
 			// was updated); however, if this is a native class (like when hot-
 			// reloading), then we want to make sure to update the skel as well
+			const bool bSkipGC = true;
 			const bool bSkeletonUpToDate = !ClassToReinstance->HasAnyClassFlags(CLASS_Native);
-			FKismetEditorUtilities::CompileBlueprint(BP, false, bSkipGarbageCollection, false, nullptr, bSkeletonUpToDate, false);
+			FKismetEditorUtilities::CompileBlueprint(BP, false, bSkipGC, false, nullptr, bSkeletonUpToDate, false);
 		}
 		else if (IsReinstancingSkeleton())
 		{
diff --git a/Engine/Source/Editor/UnrealEd/Public/DragAndDrop/AssetPathDragDropOp.h b/Engine/Source/Editor/UnrealEd/Public/DragAndDrop/AssetPathDragDropOp.h
--- a/Engine/Source/Editor/UnrealEd/Public/DragAndDrop/AssetPathDragDropOp.h
+++ b/Engine/Source/Editor/UnrealEd/Public/DragAndDrop/AssetPathDragDropOp.h
@@ -34,19 +34,13 @@ public:
 public:
 	FText GetDecoratorText() const
 	{
-		if ( CurrentHoverText.IsEmpty() )
+		if ( CurrentHoverText.IsEmpty() && PathNames.Num() > 0 )
 		{
-			FString Text = PathNames.Num() > 0 ? PathNames[0] : TEXT("");
-
-			if ( PathNames.Num() > 1 )
-			{
-				Text += TEXT(" ");
-				Text += FString::Printf(*NSLOCTEXT("ContentBrowser", "FolderDescription", "and %d other(s)").ToString(), PathNames.Num() - 1);
-			}
-
-			return FText::FromString(Text);
+			return (PathNames.Num() == 1)
+				? FText::FromString(PathNames[0])
+				: FText::Format(NSLOCTEXT("ContentBrowser", "FolderDescriptionMulti", "{0} and {1} {1}|plural(one=other,other=others)"), FText::FromString(PathNames[0]), PathNames.Num() - 1);
 		}
-		
+
 		return CurrentHoverText;
 	}
 
diff --git a/Engine/Source/Editor/UnrealEd/Public/DragAndDrop/CollectionDragDropOp.h b/Engine/Source/Editor/UnrealEd/Public/DragAndDrop/CollectionDragDropOp.h
--- a/Engine/Source/Editor/UnrealEd/Public/DragAndDrop/CollectionDragDropOp.h
+++ b/Engine/Source/Editor/UnrealEd/Public/DragAndDrop/CollectionDragDropOp.h
@@ -39,9 +39,9 @@ public:
 		{
 			return (Collections.Num() == 1)
 				? FText::FromName(Collections[0].Name)
-				: FText::Format(NSLOCTEXT("ContentBrowser", "CollectionDragDropDescription", "{0} and {1} other(s)"), FText::FromName(Collections[0].Name), FText::AsNumber(Collections.Num() - 1));
+				: FText::Format(NSLOCTEXT("ContentBrowser", "CollectionDragDropDescription", "{0} and {1} {1}|plural(one=other,other=others)"), FText::FromName(Collections[0].Name), Collections.Num() - 1);
 		}
-		
+
 		return CurrentHoverText;
 	}
 
diff --git a/Engine/Source/Editor/UnrealEd/Public/Kismet2/KismetReinstanceUtilities.h b/Engine/Source/Editor/UnrealEd/Public/Kismet2/KismetReinstanceUtilities.h
--- a/Engine/Source/Editor/UnrealEd/Public/Kismet2/KismetReinstanceUtilities.h
+++ b/Engine/Source/Editor/UnrealEd/Public/Kismet2/KismetReinstanceUtilities.h
@@ -70,9 +70,6 @@ protected:
 	/** Whether or not this resinstancer has already reinstanced  */
 	bool bHasReinstanced;
 
-	/** Don't call GC */
-	bool bSkipGarbageCollection;
-
 	/** Cached value, mostly used to determine if we're explicitly targeting the skeleton class or not */
 	enum EReinstClassType
 	{
@@ -161,7 +158,6 @@ protected:
 		, DuplicatedClass(NULL)
 		, OriginalCDO(NULL)
 		, bHasReinstanced(false)
-		, bSkipGarbageCollection(false)
 		, ReinstClassType(RCT_Unknown)
 		, ClassToReinstanceDefaultValuesCRC(0)
 		, bIsRootReinstancer(false)
diff --git a/Engine/Source/Programs/CrashReporter/CompressionHelper/Compression.cpp b/Engine/Source/Programs/CrashReporter/CompressionHelper/Compression.cpp
--- a/Engine/Source/Programs/CrashReporter/CompressionHelper/Compression.cpp
+++ b/Engine/Source/Programs/CrashReporter/CompressionHelper/Compression.cpp
@@ -79,21 +79,30 @@ extern "C"
 	__declspec(dllexport)
 	int32 __cdecl UE4CompressFileGZIP(const char* Path, const void* UncompressedBuffer, int32 UncompressedSize)
 	{
+		int32 ReturnVal = Z_ERRNO;
+
 		// Zlib wants to use unsigned long.
 		unsigned ZUncompressedSize = UncompressedSize;
 
 		gzFile FilePtr = gzopen(Path, "wb");
 		if (FilePtr == nullptr)
 		{
-			return Z_ERRNO;
+			return ReturnVal;
 		}
 
 		int ZCompressedSize = gzwrite(FilePtr, UncompressedBuffer, ZUncompressedSize);
+		
 		if (ZCompressedSize == 0)
 		{
-			return Z_ERRNO;
+			// Compression error in gzwrite()
+			gzerror(FilePtr, &ReturnVal);
+			gzclose(FilePtr);
+		}
+		else
+		{
+			ReturnVal = gzclose(FilePtr);
 		}
 
-		return gzclose(FilePtr);
+		return ReturnVal;
 	}
 }
\ No newline at end of file
diff --git a/Engine/Source/Programs/CrashReporter/CrashReportProcess/Config.cs b/Engine/Source/Programs/CrashReporter/CrashReportProcess/Config.cs
--- a/Engine/Source/Programs/CrashReporter/CrashReportProcess/Config.cs
+++ b/Engine/Source/Programs/CrashReporter/CrashReportProcess/Config.cs
@@ -332,6 +332,12 @@ namespace Tools.CrashReporter.CrashReportProcess
 		[XmlElement]
 		public int AddCrashRequestTimeoutSeconds { get; set; }
 
+		/// <summary>
+		/// Number of times to retry AddCrash after a bad response (doesn't count timeouts that are always retried).
+		/// </summary>
+		[XmlElement]
+		public int AddCrashRejectedRetries { get; set; }
+
 		/// <summary>
 		/// Time that we wait between a failed AddCrash call and a retry.
 		/// </summary>
diff --git a/Engine/Source/Programs/CrashReporter/CrashReportProcess/Config/CrashReportProcess.config b/Engine/Source/Programs/CrashReporter/CrashReportProcess/Config/CrashReportProcess.config
--- a/Engine/Source/Programs/CrashReporter/CrashReportProcess/Config/CrashReportProcess.config
+++ b/Engine/Source/Programs/CrashReporter/CrashReportProcess/Config/CrashReportProcess.config
@@ -53,6 +53,7 @@
 	<ProcessedReportsIndexPath>MyProcessedReportsIndexFilePath</ProcessedReportsIndexPath>
 	<ReportsIndexRetentionDays>120</ReportsIndexRetentionDays>
 	<AddCrashRequestTimeoutSeconds>30</AddCrashRequestTimeoutSeconds>
+	<AddCrashRejectedRetries>3</AddCrashRejectedRetries>
 	<AddCrashRetryDelaySeconds>15</AddCrashRetryDelaySeconds>
 	<FailedWebAddAlertTimeSeconds>300</FailedWebAddAlertTimeSeconds>
 	<DiskSpaceAlertPercent>10.0</DiskSpaceAlertPercent>
diff --git a/Engine/Source/Programs/CrashReporter/CrashReportProcess/CrashReportProcess.csproj b/Engine/Source/Programs/CrashReporter/CrashReportProcess/CrashReportProcess.csproj
--- a/Engine/Source/Programs/CrashReporter/CrashReportProcess/CrashReportProcess.csproj
+++ b/Engine/Source/Programs/CrashReporter/CrashReportProcess/CrashReportProcess.csproj
@@ -57,7 +57,7 @@
     <Optimize>true</Optimize>
     <DebugType>pdbonly</DebugType>
     <PlatformTarget>AnyCPU</PlatformTarget>
-    <RunCodeAnalysis>false</RunCodeAnalysis>
+    <RunCodeAnalysis>true</RunCodeAnalysis>
     <ErrorReport>prompt</ErrorReport>
     <CodeAnalysisRuleSet>MinimumRecommendedRules.ruleset</CodeAnalysisRuleSet>
   </PropertyGroup>
diff --git a/Engine/Source/Programs/CrashReporter/CrashReportProcess/CrashReportProcessService.cs b/Engine/Source/Programs/CrashReporter/CrashReportProcess/CrashReportProcessService.cs
--- a/Engine/Source/Programs/CrashReporter/CrashReportProcess/CrashReportProcessService.cs
+++ b/Engine/Source/Programs/CrashReporter/CrashReportProcess/CrashReportProcessService.cs
@@ -261,35 +261,48 @@ namespace Tools.CrashReporter.CrashReportProcess
 			bDisposing = true;
 
 			// Clean up the directory watcher and crash processor threads
+			WriteEvent("Shutdown: Stopping ReportProcessors...");
 			foreach (var Processor in Processors)
 			{
 				Processor.RequestStop();
 			}
+			WriteEvent("Shutdown: Disposing ReportProcessors...");
 			foreach (var Processor in Processors)
 			{
 				Processor.Dispose();
 			}
 			Processors.Clear();
+			WriteEvent("Shutdown: ReportProcessors stopped and disposed.");
 
+			WriteEvent("Shutdown: Disposing ReportWatcher...");
 			Watcher.Dispose();
 			Watcher = null;
+			WriteEvent("Shutdown: ReportWatcher disposed.");
 
+			WriteEvent("Shutdown: Writing ReportIndex.");
 			ReportIndex.WriteToFile();
 			ReportIndex = null;
+			WriteEvent("Shutdown: ReportIndex written.");
 
+			WriteEvent("Shutdown: Disposing AmazonClients...");
 			OutputAWS.Dispose();
 			OutputAWS = null;
-
 			DataRouterAWS.Dispose();
 			DataRouterAWS = null;
+			WriteEvent("Shutdown: AmazonClients disposed.");
 
+			WriteEvent("Shutdown: Disposing StatusReporter...");
 			StatusReporter.Dispose();
 			StatusReporter = null;
+			WriteEvent("Shutdown: StatusReporter disposed.");
 
+			WriteEvent("Shutdown: Disposing SlackWriter...");
 			Slack.Dispose();
 			Slack = null;
+			WriteEvent("Shutdown: SlackWriter disposed.");
 
 			// Flush the log to disk
+			WriteEvent("Shutdown: Disposing LogWriter.");
 			Log.Dispose();
 			Log = null;
 
diff --git a/Engine/Source/Programs/CrashReporter/CrashReportProcess/DataRouterReportQueue.cs b/Engine/Source/Programs/CrashReporter/CrashReportProcess/DataRouterReportQueue.cs
--- a/Engine/Source/Programs/CrashReporter/CrashReportProcess/DataRouterReportQueue.cs
+++ b/Engine/Source/Programs/CrashReporter/CrashReportProcess/DataRouterReportQueue.cs
@@ -178,16 +178,13 @@ namespace Tools.CrashReporter.CrashReportProcess
 					string CrashFolderPath = Path.Combine(InLandingZone, CrashFolderName);
 
 					// Early check for duplicate processed report
-					lock (ReportIndexLock)
+					if (CrashReporterProcessServicer.ReportIndex.ContainsReport(CrashFolderName))
 					{
-						if (CrashReporterProcessServicer.ReportIndex.ContainsReport(CrashFolderName))
-						{
-							// Crash report not accepted by index
-							CrashReporterProcessServicer.WriteEvent(string.Format(
-								"DataRouterReportQueue: Duplicate report skipped early {0} in a DataRouterReportQueue", CrashFolderPath));
-							CrashReporterProcessServicer.StatusReporter.IncrementCount(StatusReportingEventNames.DuplicateRejected);
-							return false; // this isn't an error so don't set error message
-						}
+						// Crash report not accepted by index
+						CrashReporterProcessServicer.WriteEvent(string.Format(
+							"DataRouterReportQueue: Duplicate report skipped early {0} in a DataRouterReportQueue", CrashFolderPath));
+						CrashReporterProcessServicer.StatusReporter.IncrementCount(StatusReportingEventNames.DuplicateRejected);
+						return false; // this isn't an error so don't set error message
 					}
 
 					// Create target folder
@@ -230,16 +227,13 @@ namespace Tools.CrashReporter.CrashReportProcess
 					string CrashFolderPath = Path.Combine(InLandingZone, CrashFolderName);
 
 					// Early check for duplicate processed report
-					lock (ReportIndexLock)
+					if (CrashReporterProcessServicer.ReportIndex.ContainsReport(CrashFolderName))
 					{
-						if (CrashReporterProcessServicer.ReportIndex.ContainsReport(CrashFolderName))
-						{
-							// Crash report not accepted by index
-							CrashReporterProcessServicer.WriteEvent(string.Format(
-								"DataRouterReportQueue: Duplicate report skipped early {0} in a DataRouterReportQueue", CrashFolderPath));
-							CrashReporterProcessServicer.StatusReporter.IncrementCount(StatusReportingEventNames.DuplicateRejected);
-							return false; // this isn't an error so don't set error message
-						}
+						// Crash report not accepted by index
+						CrashReporterProcessServicer.WriteEvent(string.Format(
+							"DataRouterReportQueue: Duplicate report skipped early {0} in a DataRouterReportQueue", CrashFolderPath));
+						CrashReporterProcessServicer.StatusReporter.IncrementCount(StatusReportingEventNames.DuplicateRejected);
+						return false; // this isn't an error so don't set error message
 					}
 
 					// Create target folder
diff --git a/Engine/Source/Programs/CrashReporter/CrashReportProcess/ReportIndex.cs b/Engine/Source/Programs/CrashReporter/CrashReportProcess/ReportIndex.cs
--- a/Engine/Source/Programs/CrashReporter/CrashReportProcess/ReportIndex.cs
+++ b/Engine/Source/Programs/CrashReporter/CrashReportProcess/ReportIndex.cs
@@ -12,33 +12,45 @@ namespace Tools.CrashReporter.CrashReportProcess
 		{
 			if (!IsEnabled) return;
 
-			try
+			lock (FileLock)
 			{
-				if (File.Exists(Filepath))
-				{
-					File.Delete(Filepath + ".backup");
-					File.Move(Filepath, Filepath + ".backup");
-				}
-
-				using (var Writer = File.CreateText(Filepath))
+				try
 				{
-					var Today = DateTime.UtcNow.Date;
+					if (File.Exists(Filepath))
+					{
+						File.Delete(Filepath + ".backup");
+						File.Move(Filepath, Filepath + ".backup");
+					}
 
-					foreach (var Item in Index)
+					using (MemoryStream MemStream = new MemoryStream())
 					{
-						if (Today - Item.Value.Date <= Retention)
+						using (StreamWriter Writer = new StreamWriter(MemStream))
 						{
-							Writer.WriteLine(ItemToString(Item));
+							var Today = DateTime.UtcNow.Date;
+
+							lock (IndexLock)
+							{
+								foreach (var Item in Index)
+								{
+									if (Today - Item.Value.Date <= Retention)
+									{
+										Writer.WriteLine(ItemToString(Item));
+									}
+								}
+								Writer.Flush();
+							}
+
+							File.WriteAllBytes(Filepath, MemStream.ToArray());
 						}
 					}
+					LastFlush = DateTime.UtcNow;
+				}
+				catch (Exception Ex)
+				{
+					CrashReporterProcessServicer.WriteException(string.Format("Failed to write ReportIndex to {0}. Exception was: {1}", Filepath, Ex), Ex);
+					CrashReporterProcessServicer.StatusReporter.Alert("ReportIndex.WriteToFile", string.Format("Failed to write ReportIndex to {0}", Filepath),
+					                                                  Config.Default.SlackAlertRepeatMinimumMinutes);
 				}
-
-				LastFlush = DateTime.UtcNow;
-			}
-			catch (Exception Ex)
-			{
-				CrashReporterProcessServicer.WriteException(string.Format("Failed to write ReportIndex to {0}. Exception was: {1}", Filepath, Ex), Ex);
-				CrashReporterProcessServicer.StatusReporter.Alert("ReportIndex.WriteToFile", string.Format("Failed to write ReportIndex to {0}", Filepath), Config.Default.SlackAlertRepeatMinimumMinutes);
 			}
 		}
 
@@ -48,43 +60,52 @@ namespace Tools.CrashReporter.CrashReportProcess
 
 			try
 			{
-				Index = new Dictionary<string, DateTime>();
-
-				if (!File.Exists(Filepath))
+				lock (IndexLock)
 				{
-					if (File.Exists(Filepath + ".backup"))
-					{
-						CrashReporterProcessServicer.WriteFailure(string.Format("Failed to read ReportIndex from {0}. Attempting to read from {1}", Filepath, Filepath + ".backup"));
-						CrashReporterProcessServicer.StatusReporter.Alert("ReportIndex.ReadFromFile.UsingBackup", string.Format("Failed to read ReportIndex from {0}. Using backup.", Filepath), Config.Default.SlackAlertRepeatMinimumMinutes);
-						File.Move(Filepath + ".backup", Filepath);
-					}
-					else
-					{
-						CrashReporterProcessServicer.WriteFailure(string.Format("Failed to read ReportIndex from {0}. Generating new one.", Filepath));
-						File.Create(Filepath).Close();
-					}
-				}
+					Index = new Dictionary<string, DateTime>();
 
-				using (var Reader = File.OpenText(Filepath))
-				{
-					string ItemStringRaw;
-					while ((ItemStringRaw = Reader.ReadLine()) != null)
+					lock (FileLock)
 					{
-						string ItemString = ItemStringRaw.Trim();
-						if (!string.IsNullOrWhiteSpace(ItemString))
+						if (!File.Exists(Filepath))
 						{
-							KeyValuePair<string, DateTime> NewItem;
-							if (!TryParseItemString(ItemString, out NewItem))
+							if (File.Exists(Filepath + ".backup"))
+							{
+								CrashReporterProcessServicer.WriteFailure(string.Format("Failed to read ReportIndex from {0}. Attempting to read from {1}", Filepath,
+																						Filepath + ".backup"));
+								CrashReporterProcessServicer.StatusReporter.Alert("ReportIndex.ReadFromFile.UsingBackup",
+																				  string.Format("Failed to read ReportIndex from {0}. Using backup.", Filepath),
+																				  Config.Default.SlackAlertRepeatMinimumMinutes);
+								File.Move(Filepath + ".backup", Filepath);
+							}
+							else
 							{
-								CrashReporterProcessServicer.WriteFailure(string.Format("Failed to read line from ReportIndex: {0}.", ItemString));
-								continue;
+								CrashReporterProcessServicer.WriteFailure(string.Format("Failed to read ReportIndex from {0}. Generating new one.", Filepath));
+								File.Create(Filepath).Close();
 							}
-							Index.Add(NewItem.Key, NewItem.Value);
 						}
-					}
-				}
 
-				LastFlush = DateTime.UtcNow;
+						using (var Reader = File.OpenText(Filepath))
+						{
+							string ItemStringRaw;
+							while ((ItemStringRaw = Reader.ReadLine()) != null)
+							{
+								string ItemString = ItemStringRaw.Trim();
+								if (!string.IsNullOrWhiteSpace(ItemString))
+								{
+									KeyValuePair<string, DateTime> NewItem;
+									if (!TryParseItemString(ItemString, out NewItem))
+									{
+										CrashReporterProcessServicer.WriteFailure(string.Format("Failed to read line from ReportIndex: {0}.", ItemString));
+										continue;
+									}
+									Index.Add(NewItem.Key, NewItem.Value);
+								}
+							}
+						}
+
+						LastFlush = DateTime.UtcNow;
+					} // lock (FileLock)
+				} // lock (IndexLock)
 			}
 			catch (Exception Ex)
 			{
@@ -97,16 +118,23 @@ namespace Tools.CrashReporter.CrashReportProcess
 		{
 			if (!IsEnabled) return true;
 
-			if (ContainsReport(ReportKey))
+			lock (IndexLock)
 			{
-				return false;
-			}
+				if (Index.ContainsKey(ReportKey))
+				{
+					return false;
+				}
 
-			Index.Add(ReportKey, DateTime.UtcNow);
+				Index.Add(ReportKey, DateTime.UtcNow);
+			}
 
-			if (LastFlush.Date < DateTime.UtcNow.Date)
+			lock (FileLock)
 			{
-				WriteToFile();
+				if (LastFlush.Date < DateTime.UtcNow.Date)
+				{
+					CrashReporterProcessServicer.WriteEvent("ReportIndex.TryAddNewReport flushing index to disk on schedule.");
+					WriteToFile();
+				}
 			}
 
 			return true;
@@ -116,14 +144,20 @@ namespace Tools.CrashReporter.CrashReportProcess
 		{
 			if (!IsEnabled) return false;
 
-			return Index.ContainsKey(ReportKey);
+			lock (IndexLock)
+			{
+				return Index.ContainsKey(ReportKey);
+			}
 		}
 
 		public bool TryRemoveReport(string ReportKey)
 		{
 			if (!IsEnabled) return true;
 
-			return Index.Remove(ReportKey);
+			lock (IndexLock)
+			{
+				return Index.Remove(ReportKey);
+			}
 		}
 
 		private static string ItemToString(KeyValuePair<string, DateTime> Item)
@@ -164,5 +198,7 @@ namespace Tools.CrashReporter.CrashReportProcess
 
 		private Dictionary<string, DateTime> Index = new Dictionary<string, DateTime>();
 		private DateTime LastFlush = DateTime.MinValue;
+		private Object IndexLock = new Object();
+		private Object FileLock = new Object();
 	}
 }
diff --git a/Engine/Source/Programs/CrashReporter/CrashReportProcess/ReportProcessor.cs b/Engine/Source/Programs/CrashReporter/CrashReportProcess/ReportProcessor.cs
--- a/Engine/Source/Programs/CrashReporter/CrashReportProcess/ReportProcessor.cs
+++ b/Engine/Source/Programs/CrashReporter/CrashReportProcess/ReportProcessor.cs
@@ -21,7 +21,22 @@ namespace Tools.CrashReporter.CrashReportProcess
 	/// A class to handle processing of received crash reports.
 	/// </summary>
 	sealed class ReportProcessor : IDisposable
-	{
+	{ 
+		class UploadRetryState
+		{
+			public bool bUploaded = false;
+			public bool bShouldRetry = true;
+			public int Timeouts = 0;
+			public int BadResponses = 0;
+		}
+
+		enum AddReportResult
+		{
+			Failed,
+			Cancelled,
+			Added
+		}
+
 		/// <summary>Absolute path to Perforce (default installation assumed)</summary>
 		const string PerforceExePath = "C:\\Program Files\\Perforce\\p4.exe";
 
@@ -79,8 +94,22 @@ namespace Tools.CrashReporter.CrashReportProcess
 			{
 				CancelSource.Cancel();
 			}
+
+			CrashReporterProcessServicer.WriteEvent("Shutdown: Stopping ReportProcessor thread...");
 			ProcessorTask.Wait();
 			ProcessorTask.Dispose();
+			CrashReporterProcessServicer.WriteEvent("Shutdown: ReportProcessor thread stopped.");
+
+			CrashReporterProcessServicer.WriteEvent("Shutdown: Stopping ReportProcessor's AddReport threads...");
+			foreach (Task AddReportTask in AddReportTasks)
+			{
+				if (AddReportTask != null)
+				{
+					AddReportTask.Wait();
+					AddReportTask.Dispose();
+				}
+			}
+			CrashReporterProcessServicer.WriteEvent("Shutdown: AddReport threads stopped.");
 
 			CancelSource.Dispose();
 		}
@@ -91,10 +120,9 @@ namespace Tools.CrashReporter.CrashReportProcess
 		/// <remarks>All exceptions are caught and written to the event log.</remarks>
 		private void Init()
 		{
-			var Cancel = CancelSource.Token;
 			ProcessorTask = new Task(() =>
 			                         {
-				                         while (!Cancel.IsCancellationRequested)
+				                         while (!CancelSource.IsCancellationRequested)
 				                         {
 					                         try
 					                         {
@@ -149,18 +177,23 @@ namespace Tools.CrashReporter.CrashReportProcess
 		/// Finalizes report files. 
 		/// Thread-safe.
 		/// </summary>
-		private void FinalizeReport(bool bAdded, DirectoryInfo DirInfo, FGenericCrashContext NewContext)
+		private void FinalizeReport(AddReportResult AddResult, DirectoryInfo DirInfo, FGenericCrashContext NewContext)
 		{
 			// Only remove if we added the report, otherwise leave all files to further investigation.
-			if (bAdded)
+			if (AddResult == AddReportResult.Added)
 			{
 				// Remove the report files as we're done with them.
 				CleanReport(DirInfo);
 			}
-			else
+			else if (AddResult == AddReportResult.Failed)
 			{
 				MoveReportToInvalid(NewContext.CrashDirectory, ReportQueueBase.GetSafeFilename(NewContext.GetAsFilename()));
 			}
+			else // AddResult == AddReportResult.Cancelled
+			{
+				// Remove report from index of completed reports and leave on disk
+				CrashReporterProcessServicer.ReportIndex.TryRemoveReport(NewContext.CrashDirectory);
+			}
 		}
 
 		/// <summary>
@@ -425,41 +458,60 @@ namespace Tools.CrashReporter.CrashReportProcess
 						RequestString = "http://localhost:80/Crashes/AddCrash/-1"; 
 					}
 
-					int Attempt = 1;
-					Func<bool> TryUpload = () =>
+				
+					Action<UploadRetryState> TryUpload = InRetryState =>
 					{
 						string ResponseString = SimpleWebRequest.GetWebServiceResponse(RequestString, Payload, Config.Default.AddCrashRequestTimeoutSeconds * 1000);
 						if (ResponseString.Length > 0)
 						{
 							// Convert response into a string
 							CrashReporterResult Result = XmlHandler.FromXmlString<CrashReporterResult>(ResponseString);
-							if (Result.ID > 0)
+							if (!Result.bSuccess || Result.ID <= 0)
+							{
+								// Website responded - fail condition
+								InRetryState.BadResponses++;
+								CrashReporterProcessServicer.WriteEvent(string.Format("PROC-{0} UploadCrash response invalid. FailedResponses={1} Response={2}",
+																					  ProcessorIndex, InRetryState.BadResponses, Result.Message));
+								if (InRetryState.BadResponses > Config.Default.AddCrashRejectedRetries)
+								{
+									InRetryState.bShouldRetry = false;
+									CrashReporterProcessServicer.WriteEvent(string.Format("PROC-{0} UploadCrash abandoned. FailedResponses={1}",
+																						  ProcessorIndex, InRetryState.BadResponses));
+								}
+							}
+							else
 							{
+								// Website responded - success
 								NewID = Result.ID;
-								return true;
+								InRetryState.bUploaded = true;
+								return;
 							}
-
-							CrashReporterProcessServicer.WriteEvent(string.Format("PROC-{0} UploadCrash response invalid (upload attempt {1}): {2}", ProcessorIndex, Attempt, Result.Message));
 						}
-						return false;
+						else
+						{
+							// No response
+							InRetryState.Timeouts++;
+						}
+						InRetryState.bUploaded = false;
 					};
-					Action<bool> OnRetry = bUploaded =>
+					Action<UploadRetryState> OnRetry = InRetryState =>
 					{
-						if (!bUploaded)
+						if (!InRetryState.bUploaded)
 						{
 							Thread.Sleep(Config.Default.AddCrashRetryDelaySeconds*1000);
-							CrashReporterProcessServicer.WriteEvent(string.Format("PROC-{0} UploadCrash retrying upload attempt {1}", ProcessorIndex, Attempt));
-							Attempt++;
+							CrashReporterProcessServicer.WriteEvent(string.Format("PROC-{0} UploadCrash retrying upload. Timeouts={1} FailedResponses={2}", ProcessorIndex, InRetryState.Timeouts, InRetryState.BadResponses));
 						}
 					};
 
-					for (bool bUploaded = false; !bUploaded; OnRetry(bUploaded))
+					UploadRetryState RetryState = new UploadRetryState();
+
+					for (; !RetryState.bUploaded && RetryState.bShouldRetry && !CancelSource.IsCancellationRequested; OnRetry(RetryState))
 					{
-						bUploaded = TryUpload();
+						TryUpload(RetryState);
 						CrashReporterProcessServicer.StatusReporter.AlertOnConsecutiveFails("AddCrash",
 						                                                                    "Cannot contact CR website. CRP is paused!",
 						                                                                    "Contact with CR website restored.",
-						                                                                    TimeSpan.FromSeconds(Config.Default.FailedWebAddAlertTimeSeconds), bUploaded);
+						                                                                    TimeSpan.FromSeconds(Config.Default.FailedWebAddAlertTimeSeconds), RetryState.bUploaded);
 					}
 				}
 #if DEBUG
@@ -488,7 +540,7 @@ namespace Tools.CrashReporter.CrashReportProcess
 		/// <param name="LogFileName">The file name of the log file in the report.</param>
 		/// <param name="DumpFileName">The file name of the minidump in the report.</param>
 		/// <param name="VideoFileName">The file name of the video file in the report, or null if there isn't one.</param>
-		private bool AddReport( DirectoryInfo DirInfo, FGenericCrashContext NewContext, string LogFileName, string DumpFileName, string VideoFileName )
+		private AddReportResult AddReport( DirectoryInfo DirInfo, FGenericCrashContext NewContext, string LogFileName, string DumpFileName, string VideoFileName )
 		{
 			try
 			{
@@ -499,7 +551,7 @@ namespace Tools.CrashReporter.CrashReportProcess
 				if (CrashDetails == "")
 				{
 					CrashReporterProcessServicer.WriteFailure(string.Format("PROC-{0} ", ProcessorIndex) + "! CreateCrash no payload: Path=" + NewContext.CrashDirectory);
-					return false;
+					return AddReportResult.Failed;
 				}
 
 				// Upload the crash to the database, and retrieve the new row id
@@ -507,9 +559,18 @@ namespace Tools.CrashReporter.CrashReportProcess
 
 				if (ReportID <= 0)
 				{
-					// Upload failed
-					CrashReporterProcessServicer.WriteFailure(string.Format("PROC-{0} ", ProcessorIndex) + "! NoUpload: Path=" + NewContext.CrashDirectory);
-					return false;
+					if (CancelSource.IsCancellationRequested)
+					{
+						// Cancelled upload
+						CrashReporterProcessServicer.WriteEvent(string.Format("PROC-{0} ", ProcessorIndex) + "CancelUpload: Path=" + NewContext.CrashDirectory);
+						return AddReportResult.Cancelled;
+					}
+					else
+					{
+						// Upload failed
+						CrashReporterProcessServicer.WriteFailure(string.Format("PROC-{0} ", ProcessorIndex) + "! NoUpload: Path=" + NewContext.CrashDirectory);
+						return AddReportResult.Failed;
+					}
 				}
 
 				bool bToS3 = Config.Default.CrashFilesToAWS && CrashReporterProcessServicer.OutputAWS.IsS3Valid;
@@ -639,14 +700,15 @@ namespace Tools.CrashReporter.CrashReportProcess
 					string.Format(
 						"AddReport: Ratio={0,2} Processed={1,7} WebAdded={2,7} AddReportTime={3} AddedPerDay={4} AddedPerMinute={5:N1}", (int) Ratio,
 						ProcessedReportCount, WebAddCounter.TotalEvents, AddReportTime.Elapsed.TotalSeconds.ToString("0.00"), (int)AddedPerDay, WebAddCounter.EventsPerSecond * 60));
-				return true;
+
+				return AddReportResult.Added;
 			}
 			catch( Exception Ex )
 			{
 				CrashReporterProcessServicer.WriteException(string.Format("PROC-{0} ", ProcessorIndex) + "AddReport: " + DirInfo.Name + "\n\n" + Ex, Ex);
 			}
 
-			return false;
+			return AddReportResult.Failed;
 		}
 
 		private static void UploadFileToS3(FileInfo FileInfo, string DestFilepath, bool bCompressed, string CompressedSuffix = null)
@@ -826,8 +888,8 @@ namespace Tools.CrashReporter.CrashReportProcess
 				NewContext.ToFile();
 				AddReportTasks[AddReportTaskSlot] = Task.Factory.StartNew(() =>
 				{
-					bool bAdded = AddReport( DirInfo, NewContext, LogFileName, DumpFileName, VideoFileName );
-					FinalizeReport( bAdded, DirInfo, NewContext );
+					AddReportResult Result = AddReport( DirInfo, NewContext, LogFileName, DumpFileName, VideoFileName );
+					FinalizeReport(Result, DirInfo, NewContext );
 				} );
 
 				return WaitTime;
diff --git a/Engine/Source/Programs/CrashReporter/CrashReportProcess/ReportQueue.cs b/Engine/Source/Programs/CrashReporter/CrashReportProcess/ReportQueue.cs
--- a/Engine/Source/Programs/CrashReporter/CrashReportProcess/ReportQueue.cs
+++ b/Engine/Source/Programs/CrashReporter/CrashReportProcess/ReportQueue.cs
@@ -269,6 +269,12 @@ namespace Tools.CrashReporter.CrashReportProcess
 			OutCrashContext.PrimaryCrashProperties.BaseDir = ReportData.BaseDir;
 			// 			OutCrashContext.PrimaryCrashProperties.RootDir
 			OutCrashContext.PrimaryCrashProperties.MachineId = ReportData.MachineId;
+			OutCrashContext.PrimaryCrashProperties.LoginId = ReportData.LoginId;
+			if (string.IsNullOrWhiteSpace(OutCrashContext.PrimaryCrashProperties.MachineId))
+			{
+				// Set MachineId from LoginId if the metadata is too new to contain MachineId
+				OutCrashContext.PrimaryCrashProperties.MachineId = OutCrashContext.PrimaryCrashProperties.LoginId;
+			}
 			OutCrashContext.PrimaryCrashProperties.EpicAccountId = ReportData.EpicAccountId;
 			// 			OutCrashContext.PrimaryCrashProperties.CallStack
 			// 			OutCrashContext.PrimaryCrashProperties.SourceContext
@@ -494,17 +500,14 @@ namespace Tools.CrashReporter.CrashReportProcess
 			// This can be safely discontinued when all crashes come from the same SQS
 			// This DOES NOT scale to multiple CRP instances
 			// ReportIndex can be disabled by leaving the path empty in config
-			lock (ReportIndexLock)
+			if (!CrashReporterProcessServicer.ReportIndex.TryAddNewReport(ReportName))
 			{
-				if (!CrashReporterProcessServicer.ReportIndex.TryAddNewReport(ReportName))
-				{
-					// Crash report not accepted by index
-					CrashReporterProcessServicer.WriteEvent(string.Format(
-						"EnqueueNewReport: Duplicate report skipped {0} in queue {1}", NewReportPath, QueueName));
-					CrashReporterProcessServicer.StatusReporter.IncrementCount(StatusReportingEventNames.DuplicateRejected);
-					ReportProcessor.CleanReport(new DirectoryInfo(NewReportPath));
-					return false;
-				}
+				// Crash report not accepted by index
+				CrashReporterProcessServicer.WriteEvent(string.Format(
+					"EnqueueNewReport: Duplicate report skipped {0} in queue {1}", NewReportPath, QueueName));
+				CrashReporterProcessServicer.StatusReporter.IncrementCount(StatusReportingEventNames.DuplicateRejected);
+				ReportProcessor.CleanReport(new DirectoryInfo(NewReportPath));
+				return false;
 			}
 
 			if (ShouldDecimateNextReport())
@@ -689,8 +692,6 @@ namespace Tools.CrashReporter.CrashReportProcess
 
 		private readonly EventCounter EnqueueCounter = new EventCounter(TimeSpan.FromMinutes(90), 20);
 
-		protected static Object ReportIndexLock = new Object();
-
 		/// <summary>
 		/// The number of waiting reports needed to trigger decimation (discarding a fraction of the incoming reports to stop runaway backlog)
 		/// </summary>
diff --git a/Engine/Source/Programs/CrashReporter/CrashReportProcess/ReportWatcher.cs b/Engine/Source/Programs/CrashReporter/CrashReportProcess/ReportWatcher.cs
--- a/Engine/Source/Programs/CrashReporter/CrashReportProcess/ReportWatcher.cs
+++ b/Engine/Source/Programs/CrashReporter/CrashReportProcess/ReportWatcher.cs
@@ -38,13 +38,18 @@ namespace Tools.CrashReporter.CrashReportProcess
 		{
 			// Cancel the task and wait for it to quit
 			CancelSource.Cancel();
+
+			CrashReporterProcessServicer.WriteEvent("Shutdown: Stopping ReportWatcher thread...");
 			WatcherTask.Wait();
 			WatcherTask.Dispose();
+			CrashReporterProcessServicer.WriteEvent("Shutdown: ReportWatcher thread stopped.");
 
+			CrashReporterProcessServicer.WriteEvent("Shutdown: Disposing ReportQueues...");
 			foreach (var Queue in ReportQueues)
 			{
 				Queue.Dispose();
 			}
+			CrashReporterProcessServicer.WriteEvent("Shutdown: ReportQueues disposed.");
 
 			CancelSource.Dispose();
 		}
@@ -135,7 +140,16 @@ namespace Tools.CrashReporter.CrashReportProcess
 
 					TimeSpan TimeTaken = DateTime.Now - StartTime;
 					CrashReporterProcessServicer.WriteEvent(string.Format("Checking Landing Zones took {0:F2} seconds", TimeTaken.TotalSeconds));
-					await Task.Delay(30000, Cancel);
+
+					try
+					{
+						await Task.Delay(30000, Cancel);
+					}
+					catch (TaskCanceledException Ex)
+					{
+						CrashReporterProcessServicer.WriteEvent("WatcherTask exiting delay because it was cancelled.");
+						break;
+					}
 				}
 			});
 		}
