PR #4932: Fix to generate cleaner C++ files when using "New C++ Class" (Contributed by [REDACTED_PERSON])

diff --git a/Engine/Source/Editor/GameProjectGeneration/Private/GameProjectUtils.cpp b/Engine/Source/Editor/GameProjectGeneration/Private/GameProjectUtils.cpp
--- a/Engine/Source/Editor/GameProjectGeneration/Private/GameProjectUtils.cpp
+++ b/Engine/Source/Editor/GameProjectGeneration/Private/GameProjectUtils.cpp
@@ -2869,9 +2869,13 @@ bool GameProjectUtils::GenerateClassHeaderFile(const FString& NewHeaderFileName,
 	FinalOutput = FinalOutput.Replace(TEXT("%UCLASS_SPECIFIER_LIST%"), *MakeCommaDelimitedList(ClassSpecifierList, false), ESearchCase::CaseSensitive);
 	FinalOutput = FinalOutput.Replace(TEXT("%PREFIXED_CLASS_NAME%"), *PrefixedClassName, ESearchCase::CaseSensitive);
 	FinalOutput = FinalOutput.Replace(TEXT("%PREFIXED_BASE_CLASS_NAME%"), *PrefixedBaseClassName, ESearchCase::CaseSensitive);
-	FinalOutput = FinalOutput.Replace(TEXT("%EVENTUAL_CONSTRUCTOR_DECLARATION%"), *EventualConstructorDeclaration, ESearchCase::CaseSensitive);
-	FinalOutput = FinalOutput.Replace(TEXT("%CLASS_PROPERTIES%"), *ClassProperties, ESearchCase::CaseSensitive);
-	FinalOutput = FinalOutput.Replace(TEXT("%CLASS_FUNCTION_DECLARATIONS%"), *ClassFunctionDeclarations, ESearchCase::CaseSensitive);
+
+	// Special case where where the wildcard starts with a tab and ends with a new line
+	const bool bLeadingTab = true;
+	const bool bTrailingNewLine = true;
+	FinalOutput = ReplaceWildcard(FinalOutput, TEXT("%EVENTUAL_CONSTRUCTOR_DECLARATION%"), *EventualConstructorDeclaration, bLeadingTab, bTrailingNewLine);
+	FinalOutput = ReplaceWildcard(FinalOutput, TEXT("%CLASS_PROPERTIES%"), *ClassProperties, bLeadingTab, bTrailingNewLine);
+	FinalOutput = ReplaceWildcard(FinalOutput, TEXT("%CLASS_FUNCTION_DECLARATIONS%"), *ClassFunctionDeclarations, bLeadingTab, bTrailingNewLine);
 	if (BaseClassIncludeDirective.Len() == 0)
 	{
 		FinalOutput = FinalOutput.Replace(TEXT("%BASE_CLASS_INCLUDE_DIRECTIVE%") LINE_TERMINATOR, TEXT(""), ESearchCase::CaseSensitive);
@@ -3028,16 +3032,16 @@ bool GameProjectUtils::GenerateClassCPPFile(const FString& NewCPPFileName, const
 	// Not all of these will exist in every class template
 	FString FinalOutput = Template.Replace(TEXT("%COPYRIGHT_LINE%"), *MakeCopyrightLine(), ESearchCase::CaseSensitive);
 	FinalOutput = FinalOutput.Replace(TEXT("%UNPREFIXED_CLASS_NAME%"), *UnPrefixedClassName, ESearchCase::CaseSensitive);
-	FinalOutput = FinalOutput.Replace(TEXT("%MODULE_NAME%"), *ModuleInfo.ModuleName, ESearchCase::CaseSensitive);
-	if (PchIncludeDirective.Len() == 0)
-	{
-		FinalOutput = FinalOutput.Replace(TEXT("%PCH_INCLUDE_DIRECTIVE%") LINE_TERMINATOR, TEXT(""), ESearchCase::CaseSensitive);
-	}
-	FinalOutput = FinalOutput.Replace(TEXT("%PCH_INCLUDE_DIRECTIVE%"), *PchIncludeDirective, ESearchCase::CaseSensitive);
 	FinalOutput = FinalOutput.Replace(TEXT("%PREFIXED_CLASS_NAME%"), *PrefixedClassName, ESearchCase::CaseSensitive);
-	FinalOutput = FinalOutput.Replace(TEXT("%EVENTUAL_CONSTRUCTOR_DEFINITION%"), *EventualConstructorDefinition, ESearchCase::CaseSensitive);
-	FinalOutput = FinalOutput.Replace(TEXT("%ADDITIONAL_MEMBER_DEFINITIONS%"), *AdditionalMemberDefinitions, ESearchCase::CaseSensitive);
-	FinalOutput = FinalOutput.Replace(TEXT("%ADDITIONAL_INCLUDE_DIRECTIVES%"), *AdditionalIncludesStr, ESearchCase::CaseSensitive);
+	FinalOutput = FinalOutput.Replace(TEXT("%MODULE_NAME%"), *ModuleInfo.ModuleName, ESearchCase::CaseSensitive);
+
+	// Special case where where the wildcard ends with a new line
+	const bool bLeadingTab = false;
+	const bool bTrailingNewLine = true;
+	FinalOutput = ReplaceWildcard(FinalOutput, TEXT("%PCH_INCLUDE_DIRECTIVE%"), *PchIncludeDirective, bLeadingTab, bTrailingNewLine);
+	FinalOutput = ReplaceWildcard(FinalOutput, TEXT("%ADDITIONAL_INCLUDE_DIRECTIVES%"), *AdditionalIncludesStr, bLeadingTab, bTrailingNewLine);
+	FinalOutput = ReplaceWildcard(FinalOutput, TEXT("%EVENTUAL_CONSTRUCTOR_DEFINITION%"), *EventualConstructorDefinition, bLeadingTab, bTrailingNewLine);
+	FinalOutput = ReplaceWildcard(FinalOutput, TEXT("%ADDITIONAL_MEMBER_DEFINITIONS%"), *AdditionalMemberDefinitions, bLeadingTab, bTrailingNewLine);
 
 	HarvestCursorSyncLocation( FinalOutput, OutSyncLocation );
 
@@ -3189,6 +3193,29 @@ bool GameProjectUtils::GeneratePluginModuleHeaderFile(const FString& HeaderFileN
 	return WriteOutputFile(HeaderFileName, FinalOutput, OutFailReason);
 }
 
+FString GameProjectUtils::ReplaceWildcard(const FString& Input, const FString& From, const FString& To, bool bLeadingTab, bool bTrailingNewLine)
+{
+	FString Result = Input;
+	FString WildCard = bLeadingTab ? TEXT("\t") : TEXT("");
+
+	WildCard.Append(From);
+
+	if (bTrailingNewLine)
+	{
+		WildCard.Append(LINE_TERMINATOR);
+	}
+
+	int32 NumReplacements = Result.ReplaceInline(*WildCard, *To, ESearchCase::CaseSensitive);
+
+	// if replacement fails, try again using just the plain wildcard without tab and/or new line
+	if (NumReplacements == 0)
+	{
+		Result = Result.Replace(*From, *To, ESearchCase::CaseSensitive);
+	}
+
+	return Result;
+}
+
 void GameProjectUtils::OnUpdateProjectConfirm()
 {
 	UpdateProject();
diff --git a/Engine/Source/Editor/GameProjectGeneration/Public/GameProjectUtils.h b/Engine/Source/Editor/GameProjectGeneration/Public/GameProjectUtils.h
--- a/Engine/Source/Editor/GameProjectGeneration/Public/GameProjectUtils.h
+++ b/Engine/Source/Editor/GameProjectGeneration/Public/GameProjectUtils.h
@@ -361,6 +361,19 @@ private:
 	/** Generates a main game module header file */
 	static bool GenerateGameModuleHeaderFile(const FString& NewGameModuleHeaderFileName, const TArray<FString>& PublicHeaderIncludes, FText& OutFailReason);
 
+	/**
+	* Replace a wildcard with another string
+	*
+	* @param Input				The input string
+	* @param From				The wildcard to be replaced
+	* @param To					The text the wildcard should be replaced with
+	* @param bLeadingTab		If the line where the wildcard is located starts with a tab
+	* @param bTrailingNewLine	If the line where the wildcard is located ends with a new line
+
+	* @return the input string after replacement of the wildcard
+	*/
+	static FString ReplaceWildcard(const FString& Input, const FString& From, const FString& To, bool bLeadingTab = false, bool bTrailingNewLine = false);
+
 	/** Handler for when the user confirms a project update */
 	static void OnUpdateProjectConfirm();
 
