Add ability to selectively target audio ini properties per platform and initialize accordingly on project creation.

diff --git a/Engine/Source/Editor/GameProjectGeneration/Private/GameProjectUtils.cpp b/Engine/Source/Editor/GameProjectGeneration/Private/GameProjectUtils.cpp
--- a/Engine/Source/Editor/GameProjectGeneration/Private/GameProjectUtils.cpp
+++ b/Engine/Source/Editor/GameProjectGeneration/Private/GameProjectUtils.cpp
@@ -66,6 +66,9 @@
 
 #include "Dialogs/SOutputLogDialog.h"
 
+#include "IAudioExtensionPlugin.h"
+#include "AudioPluginUtilities.h"
+#include "Sound/AudioSettings.h"
 #include "Sound/SoundEffectSubmix.h"
 #include "Sound/SoundEffectSource.h"
 #include "Components/SynthComponent.h"
@@ -90,17 +93,68 @@
 TWeakPtr<SNotificationItem> GameProjectUtils::UpdateGameProjectNotification = NULL;
 TWeakPtr<SNotificationItem> GameProjectUtils::WarningProjectNameNotification = NULL;
 
-FString GameProjectUtils::DefaultFeaturePackExtension(TEXT(".upack"));	
+FString GameProjectUtils::DefaultFeaturePackExtension(TEXT(".upack"));
 
 bool GameProjectUtils::bUseAudioMixerForAllPlatforms = false;
 
-TArray<FString> GameProjectUtils::AudioMixerEnabledPlatforms(
+struct FAudioDefaultPlatformSettings
 {
-	// If bUseAudioMixerForAllPlatforms is set to false,
-	// This can be used to only flag specific platforms to use the new audio engine on new projects.
-	// For example, for windows:
-	//TEXT("Windows")
-});
+	FString Name;
+	EAudioPlatform Platform;
+	FAudioPlatformSettings Settings;
+	bool bUseAudioMixer;
+
+	FAudioDefaultPlatformSettings(EAudioPlatform InPlatform)
+		: Platform(InPlatform)
+		, bUseAudioMixer(false)
+	{
+	}
+};
+
+namespace
+{
+	TMap<FString, FAudioDefaultPlatformSettings> GetAudioPlatformProjectDefaultSettings()
+	{
+		TMap<FString, FAudioDefaultPlatformSettings> DefaultProjectSettings;
+
+		// If bUseAudioMixerForAllPlatforms is set to false, uncomment the following line to enable
+		// the new audio mixer on specific platform. Ex. for Windows:
+		// WindowsSettings.bUseAudioMixer = true;
+
+		FAudioDefaultPlatformSettings AndroidSettings(EAudioPlatform::Android);
+		AndroidSettings.Settings.MaxChannels = 12;
+		DefaultProjectSettings.Add(TEXT("Android"), AndroidSettings);
+
+		FAudioDefaultPlatformSettings IOSSettings(EAudioPlatform::IOS);
+		IOSSettings.Settings.MaxChannels = 16;
+		DefaultProjectSettings.Add(TEXT("IOS"), IOSSettings);
+
+		FAudioDefaultPlatformSettings LinuxSettings(EAudioPlatform::Linux);
+		LinuxSettings.Settings.MaxChannels = 16;
+		DefaultProjectSettings.Add(TEXT("Linux"), LinuxSettings);
+
+		FAudioDefaultPlatformSettings MacSettings(EAudioPlatform::Mac);
+		DefaultProjectSettings.Add(TEXT("Mac"), MacSettings);
+
+		FAudioDefaultPlatformSettings PS4Settings(EAudioPlatform::Playstation4);
+		DefaultProjectSettings.Add(TEXT("PS4"), PS4Settings);
+
+		FAudioDefaultPlatformSettings SwitchSettings(EAudioPlatform::Switch);
+		DefaultProjectSettings.Add(TEXT("Switch"), SwitchSettings);
+
+		FAudioDefaultPlatformSettings WindowsSettings(EAudioPlatform::Windows);
+		WindowsSettings.Settings.CallbackBufferFrameSize = 256;
+		WindowsSettings.Settings.NumBuffers = 7;
+		DefaultProjectSettings.Add(TEXT("Windows"), WindowsSettings);
+
+		FAudioDefaultPlatformSettings XBoxSettings(EAudioPlatform::XboxOne);
+		XBoxSettings.Settings.CallbackBufferFrameSize = 256;
+		XBoxSettings.Settings.NumBuffers = 7;
+		DefaultProjectSettings.Add(TEXT("XboxOne"), XBoxSettings);
+
+		return MoveTemp(DefaultProjectSettings);
+	}
+} // namespace <>
 
 FText FNewClassInfo::GetClassName() const
 {
@@ -148,7 +202,7 @@ FText FNewClassInfo::GetClassDescription(const bool bFullDescription/* = true*/)
 					}
 
 					// Strip out any new-lines in the description
-					ClassDescription.ReplaceInline(TEXT("\n"), TEXT(" "));	
+					ClassDescription.ReplaceInline(TEXT("\n"), TEXT(" "));
 				}
 
 				return FText::FromString(ClassDescription);
@@ -458,7 +512,7 @@ FString FNewClassInfo::GetSourceTemplateFilename() const
 			}
 			// Some other non-actor, non-component UObject class
 			return TEXT( "UObjectClass.cpp.template" );
-	
+
 	case EClassType::EmptyCpp:
 		return TEXT("EmptyClass.cpp.template");
 
@@ -688,7 +742,7 @@ bool GameProjectUtils::OpenCodeIDE(const FString& ProjectFile, FText& OutFailRea
 void GameProjectUtils::GetStarterContentFiles(TArray<FString>& OutFilenames)
 {
 	FString const SrcFolder = FPaths::FeaturePackDir();
-	
+
 	FString SearchPath = TEXT("*");
 	SearchPath += DefaultFeaturePackExtension;
 	IFileManager::Get().FindFilesRecursive(OutFilenames, *SrcFolder, *SearchPath, /*Files=*/true, /*Directories=*/false);
@@ -961,7 +1015,7 @@ void GameProjectUtils::OpenAddToProjectDialog(const FAddToProjectConfig& Config,
 		.SizingRule( ESizingRule::FixedSize )
 		.SupportsMinimize(false) .SupportsMaximize(false);
 
-	TSharedRef<SNewClassDialog> NewClassDialog = 
+	TSharedRef<SNewClassDialog> NewClassDialog =
 		SNew(SNewClassDialog)
 		.ParentWindow(AddCodeWindow)
 		.Class(Config._ParentClass)
@@ -1041,7 +1095,7 @@ bool GameProjectUtils::IsValidClassNameForCreation(const FString& NewClassName,
 	{
 		return false;
 	}
-	
+
 	// Look for a duplicate class in memory
 	for ( TObjectIterator<UClass> ClassIt; ClassIt; ++ClassIt )
 	{
@@ -1204,7 +1258,7 @@ bool GameProjectUtils::GenerateProjectFromScratch(const FProjectInformation& InP
 	{
 		return false;
 	}
-	
+
 	// Make the Content folder
 	const FString ContentFolder = NewProjectFolder / TEXT("Content");
 	if ( !IFileManager::Get().MakeDirectory(*ContentFolder) )
@@ -1527,7 +1581,7 @@ bool GameProjectUtils::CreateProjectFromTemplate(const FProjectInformation& InPr
 					// Comment lines start with ";". Skip these lines entirely.
 					else if ( Line.StartsWith(TEXT(";")) )
 					{
-						
+
 					}
 					// If this is a section line, update the section
 					else if ( Line.StartsWith(TEXT("[")) )
@@ -1617,7 +1671,7 @@ bool GameProjectUtils::CreateProjectFromTemplate(const FProjectInformation& InPr
 	{
 		return false;
 	}
-	
+
 	SlowTask.EnterProgressFrame();
 
 	// Generate the project file
@@ -1677,7 +1731,7 @@ bool GameProjectUtils::CreateProjectFromTemplate(const FProjectInformation& InPr
 	{
 		return false;
 	}
-	
+
 	if( OutCreatedFiles != nullptr )
 	{
 		OutCreatedFiles->Append(CreatedFiles);
@@ -1796,7 +1850,7 @@ void GameProjectUtils::DeleteCreatedFiles(const FString& RootFolder, const TArra
 FString GameProjectUtils::GetHardwareConfigString(const FProjectInformation& InProjectInfo)
 {
 	FString HardwareTargeting;
-	
+
 	FString TargetHardwareAsString;
 	UEnum::GetValueAsString(TEXT("/Script/HardwareTargeting.EHardwareClass"), InProjectInfo.TargetedHardware, /*out*/ TargetHardwareAsString);
 
@@ -1827,7 +1881,7 @@ bool GameProjectUtils::GenerateConfigFiles(const FProjectInformation& InProjectI
 
 		FileContents += GetHardwareConfigString(InProjectInfo);
 		FileContents += LINE_TERMINATOR;
-		
+
 		if(bUseAudioMixerForAllPlatforms)
 		{
 			FileContents += TEXT("[Audio]") LINE_TERMINATOR;
@@ -1861,7 +1915,7 @@ bool GameProjectUtils::GenerateConfigFiles(const FProjectInformation& InProjectI
 					SpecificGameDefaultMap = TEXT("/Game/StarterContent/Maps/Minimal_Default");
 				}
 			}
-			
+
 			// Write out the settings for startup map and game default map
 			FileContents += TEXT("[/Script/EngineSettings.GameMapsSettings]") LINE_TERMINATOR;
 			FileContents += FString::Printf(TEXT("EditorStartupMap=%s") LINE_TERMINATOR, *SpecificEditorStartupMap);
@@ -1869,7 +1923,7 @@ bool GameProjectUtils::GenerateConfigFiles(const FProjectInformation& InProjectI
 			if (InProjectInfo.bShouldGenerateCode)
 			{
 				FileContents += FString::Printf(TEXT("GlobalDefaultGameMode=\"/Script/%s.%sGameMode\"") LINE_TERMINATOR, *NewProjectName, *NewProjectName);
-			}			
+			}
 		}
 
 		if (WriteOutputFile(DefaultEngineIniFilename, FileContents, OutFailReason))
@@ -1914,40 +1968,81 @@ bool GameProjectUtils::GenerateConfigFiles(const FProjectInformation& InProjectI
 		}
 	}
 
-	// inis for any audiomixer-enabled platforms:
+	// Platforms inis:
 	{
-		for (const FString& PlatformName : AudioMixerEnabledPlatforms)
+		if (!GeneratePlatformConfigFiles(InProjectInfo, OutFailReason))
 		{
-			if (!GeneratePlatformConfigFiles(InProjectInfo, PlatformName, OutFailReason))
-			{
-				return false;
-			}
+			return false;
 		}
 	}
 
 	return true;
 }
 
-bool GameProjectUtils::GeneratePlatformConfigFiles(const FProjectInformation& InProjectInfo, const FString& InPlatformName, FText& OutFailReason)
+bool GameProjectUtils::GeneratePlatformConfigFiles(const FProjectInformation& InProjectInfo, FText& OutFailReason)
 {
-	const FString NewProjectFolder = FPaths::GetPath(InProjectInfo.ProjectFilename);
+	TMap<FString, FAudioDefaultPlatformSettings> ProjectDefaults = GetAudioPlatformProjectDefaultSettings();
 
-	FString ProjectConfigPath = NewProjectFolder / TEXT("Config");
+	static const FAudioPlatformSettings DefaultSettings;
 
-	const FString PlatformEngineIniFilename = ProjectConfigPath / InPlatformName / InPlatformName + TEXT("Engine.ini");
-	FString FileContents;
+	for (TPair<FString, FAudioDefaultPlatformSettings>& SettingsPair : ProjectDefaults)
+	{
+		FString FileContents;
 
-	FileContents += TEXT("[Audio]") LINE_TERMINATOR;
-	FileContents += TEXT("UseAudioMixer=True");
+		if (bUseAudioMixerForAllPlatforms || SettingsPair.Value.bUseAudioMixer)
+		{
+			FileContents += TEXT("[Audio]") LINE_TERMINATOR;
+			FileContents += TEXT("UseAudioMixer=True") LINE_TERMINATOR;
+			FileContents += LINE_TERMINATOR;
+		}
 
-	if (WriteOutputFile(PlatformEngineIniFilename, FileContents, OutFailReason))
-	{
-		return true;
-	}
-	else
-	{
-		return false;
+		const FString& PlatformName = SettingsPair.Key;
+		const FAudioPlatformSettings& PlatformSettings = SettingsPair.Value.Settings;
+
+		FileContents += TEXT("[") + FString(AudioPluginUtilities::GetPlatformConfigSection(SettingsPair.Value.Platform)) + TEXT("]") + LINE_TERMINATOR;
+
+		if (DefaultSettings.SampleRate == PlatformSettings.SampleRate)
+		{
+			FileContents += TEXT(";");
+		}
+		FileContents += TEXT("AudioSampleRate=") + FString::Printf(TEXT("%d"), PlatformSettings.SampleRate) + LINE_TERMINATOR;
+
+		if (DefaultSettings.MaxChannels == PlatformSettings.MaxChannels)
+		{
+			FileContents += TEXT(";");
+		}
+		FileContents += TEXT("AudioMaxChannels=") + FString::Printf(TEXT("%d"), PlatformSettings.MaxChannels) + LINE_TERMINATOR;
+
+		if (DefaultSettings.CallbackBufferFrameSize == PlatformSettings.CallbackBufferFrameSize)
+		{
+			FileContents += TEXT(";");
+		}
+		FileContents += TEXT("AudioCallbackBufferFrameSize=") + FString::Printf(TEXT("%d"), PlatformSettings.CallbackBufferFrameSize) + LINE_TERMINATOR;
+
+		if (DefaultSettings.NumBuffers == PlatformSettings.NumBuffers)
+		{
+			FileContents += TEXT(";");
+		}
+		FileContents += TEXT("AudioNumBuffersToEnqueue=") + FString::Printf(TEXT("%d"), PlatformSettings.NumBuffers) + LINE_TERMINATOR;
+
+		if (DefaultSettings.NumSourceWorkers == PlatformSettings.NumSourceWorkers)
+		{
+			FileContents += TEXT(";");
+		}
+		FileContents += TEXT("AudioNumSourceWorkers=") + FString::Printf(TEXT("%d"), PlatformSettings.NumSourceWorkers) + LINE_TERMINATOR;
+
+		FileContents += LINE_TERMINATOR;
+
+		const FString NewProjectFolder = FPaths::GetPath(InProjectInfo.ProjectFilename);
+		const FString ProjectConfigPath = NewProjectFolder / TEXT("Config");
+		const FString PlatformEngineIniFilename = ProjectConfigPath / PlatformName / PlatformName + TEXT("Engine.ini");
+		if (!WriteOutputFile(PlatformEngineIniFilename, FileContents, OutFailReason))
+		{
+			return false;
+		}
 	}
+
+	return true;
 }
 
 bool GameProjectUtils::GenerateBasicSourceCode(TArray<FString>& OutCreatedFiles, FText& OutFailReason)
@@ -2084,7 +2179,7 @@ bool GameProjectUtils::GenerateGameFrameworkSourceCode(const FString& NewProject
 		const UClass* BaseClass = AGameModeBase::StaticClass();
 		const FString NewClassName = NewProjectName + BaseClass->GetName();
 		const FString NewCPPFilename = GameModulePath / NewClassName + TEXT(".cpp");
-		
+
 		TArray<FString> PropertyOverrides;
 		TArray<FString> AdditionalIncludes;
 		FString UnusedSyncLocation;
@@ -2154,7 +2249,7 @@ bool GameProjectUtils::IsStarterContentAvailableForNewProjects()
 {
 	TArray<FString> StarterContentFiles;
 	GetStarterContentFiles(StarterContentFiles);
-	
+
 	bool bHasStaterContent = StarterContentFiles.FindByPredicate([&](const FString& Str){ return Str.Contains("StarterContent"); }) != nullptr;
 	return bHasStaterContent;
 }
@@ -2787,9 +2882,9 @@ FString GameProjectUtils::DetermineModuleIncludePath(const FModuleContextInfo& M
 
 	if(FindSourceFileInProject(ModuleInfo.ModuleName + ".h", ModuleInfo.ModuleSourcePath, ModuleIncludePath))
 	{
-		// Work out where the module header is; 
+		// Work out where the module header is;
 		// if it's Public then we can include it without any path since all Public and Classes folders are on the include path
-		// if it's located elsewhere, then we'll need to include it relative to the module source root as we can't guarantee 
+		// if it's located elsewhere, then we'll need to include it relative to the module source root as we can't guarantee
 		// that other folders are on the include paths
 		EClassLocation ModuleLocation;
 		if(GetClassLocation(ModuleIncludePath, ModuleInfo, ModuleLocation))
@@ -3136,7 +3231,7 @@ bool GameProjectUtils::GeneratePluginModuleBuildFile(const FString& NewBuildFile
 	FinalOutput = FinalOutput.Replace(TEXT("%PUBLIC_DEPENDENCY_MODULE_NAMES%"), *MakeCommaDelimitedList(PublicDependencyModuleNames), ESearchCase::CaseSensitive);
 	FinalOutput = FinalOutput.Replace(TEXT("%PRIVATE_DEPENDENCY_MODULE_NAMES%"), *MakeCommaDelimitedList(PrivateDependencyModuleNames), ESearchCase::CaseSensitive);
 	FinalOutput = FinalOutput.Replace(TEXT("%MODULE_NAME%"), *ModuleName, ESearchCase::CaseSensitive);
-	
+
 	const FString PCHUsage = bUseExplicitOrSharedPCHs ? TEXT("UseExplicitOrSharedPCHs") : TEXT("UseSharedPCHs");
 	FinalOutput = FinalOutput.Replace(TEXT("%PCH_USAGE%"), *PCHUsage, ESearchCase::CaseSensitive);
 
@@ -3478,7 +3573,7 @@ void FindCodeFiles(const TCHAR* BaseDirectory, TArray<FString>& FileNames, int32
 		TArray<FString>& FileNames;
 		int32 MaxNumFileNames;
 
-		FDirectoryVisitor(TArray<FString>& InFileNames, int32 InMaxNumFileNames) 
+		FDirectoryVisitor(TArray<FString>& InFileNames, int32 InMaxNumFileNames)
 			: FileNames(InFileNames)
 			, MaxNumFileNames(InMaxNumFileNames)
 		{
@@ -3725,7 +3820,7 @@ GameProjectUtils::EAddCodeToProjectResult GameProjectUtils::AddCodeToProject_Int
 
 	// Update project file if needed.
 	auto bUpdateProjectModules = false;
-	
+
 	// If the project does not already contain code, add the primary game module
 	TArray<FString> CreatedFiles;
 	TArray<FString> StartupModuleNames;
@@ -3824,7 +3919,7 @@ GameProjectUtils::EAddCodeToProjectResult GameProjectUtils::AddCodeToProject_Int
 		// We managed the gather, so we can skip running the full generate
 		bGenerateProjectFiles = false;
 	}
-	
+
 	if ( bGenerateProjectFiles )
 	{
 		// Generate project files if we happen to be using a project file.
@@ -3933,7 +4028,7 @@ bool GameProjectUtils::FindSourceFileInProject(const FString& InFilename, const
 {
 	TArray<FString> Filenames;
 	IFileManager::Get().FindFilesRecursive(Filenames, *InSearchPath, *InFilename, true, false, false);
-	
+
 	if(Filenames.Num())
 	{
 		// Assume it's the first match (we should really only find a single file with a given name within a project anyway)
@@ -4000,7 +4095,7 @@ bool GameProjectUtils::InsertFeaturePacksIntoINIFile(const FProjectInformation&
 		}
 		PackList.Add(StarterPack);
 	}
-	
+
 	if (PackList.Num() != 0)
 	{
 		FString FileOutput;
@@ -4038,7 +4133,7 @@ bool GameProjectUtils::AddSharedContentToProject(const FProjectInformation &InPr
 
 	const FString ProjectConfigPath = DestFolder / TEXT("Config");
 	const FString IniFilename = ProjectConfigPath / TEXT("DefaultGame.ini");
-	
+
 	// Now any packs specified in the template def.
 	UTemplateProjectDefs* TemplateDefs = LoadTemplateDefs(SrcFolder);
 	if (TemplateDefs != NULL)
@@ -4057,7 +4152,7 @@ bool GameProjectUtils::AddSharedContentToProject(const FProjectInformation &InPr
 			FFormatNamedArguments Args;
 			Args.Add(TEXT("TemplateName"), FText::FromString(SrcFolder));
 			OutFailReason = FText::Format(LOCTEXT("SharedResourceError", "Error adding shared resources for '{TemplateName}'."), Args);
-			return false;		
+			return false;
 		}
 	}
 	return true;
diff --git a/Engine/Source/Editor/GameProjectGeneration/Public/GameProjectUtils.h b/Engine/Source/Editor/GameProjectGeneration/Public/GameProjectUtils.h
--- a/Engine/Source/Editor/GameProjectGeneration/Public/GameProjectUtils.h
+++ b/Engine/Source/Editor/GameProjectGeneration/Public/GameProjectUtils.h
@@ -70,7 +70,7 @@ public:
 
 		/** There were errors when adding the new source files */
 		FailedToAddCode,
-		 
+
 		/** There were errors when hot-reloading the new module */
 		FailedToHotReload,
 	};
@@ -83,7 +83,7 @@ public:
 
 		/** There were errors while duplicating project files */
 		Failed,
-		 
+
 		/** User has canceled project duplication process */
 		UserCanceled
 	};
@@ -112,8 +112,8 @@ public:
 	/** Updates the given project file to an engine identifier. Returns true if the project was updated successfully or if no update was needed */
 	static bool UpdateGameProject(const FString& ProjectFile, const FString& EngineIdentifier, FText& OutFailReason);
 
-	/** 
-	 * Opens a dialog to add code files or blueprints to the current project. 
+	/**
+	 * Opens a dialog to add code files or blueprints to the current project.
 	 *
 	 * @param	Config			Configuration options for the dialog
 	 * @param	InDomain		The domain of the class we're creating (native or blueprint)
@@ -141,7 +141,7 @@ public:
 	/** @return The number of code files in the currently loaded project */
 	static int32 GetProjectCodeFileCount();
 
-	/** 
+	/**
 	* Retrieves file and size info about the project's source directory
 	* @param OutNumFiles Contains the number of files within the source directory
 	* @param OutDirectorySize Contains the combined size of all files in the directory
@@ -175,18 +175,18 @@ public:
 	 */
 	static void ResetCurrentProjectModulesCache();
 
-	/** 
-	 * Check to see if the given path is a valid place to put source code for this project (exists within the source root path) 
+	/**
+	 * Check to see if the given path is a valid place to put source code for this project (exists within the source root path)
 	 *
 	 * @param	InPath				The path to check
 	 * @param	ModuleInfo			Information about the module being validated
 	 * @param	OutFailReason		Optional parameter to fill with failure information
-	 * 
+	 *
 	 * @return	true if the path is valid, false otherwise
 	 */
 	static bool IsValidSourcePath(const FString& InPath, const FModuleContextInfo& ModuleInfo, FText* const OutFailReason = nullptr);
 
-	/** 
+	/**
 	 * Given the path provided, work out where generated .h and .cpp files would be placed
 	 *
 	 * @param	InPath				The path to use a base
@@ -194,19 +194,19 @@ public:
 	 * @param	OutHeaderPath		The path where the .h file should be placed
 	 * @param	OutSourcePath		The path where the .cpp file should be placed
 	 * @param	OutFailReason		Optional parameter to fill with failure information
-	 * 
+	 *
 	 * @return	false if the paths are invalid
 	 */
 	static bool CalculateSourcePaths(const FString& InPath, const FModuleContextInfo& ModuleInfo, FString& OutHeaderPath, FString& OutSourcePath, FText* const OutFailReason = nullptr);
 
-	/** 
+	/**
 	 * Given the path provided, work out where it's located within the Source folder
 	 *
 	 * @param	InPath				The path to use a base
 	 * @param	ModuleInfo			Information about the module being validated
 	 * @param	OutClassLocation	The location within the Source folder
 	 * @param	OutFailReason		Optional parameter to fill with failure information
-	 * 
+	 *
 	 * @return	false if the paths are invalid
 	 */
 	static bool GetClassLocation(const FString& InPath, const FModuleContextInfo& ModuleInfo, EClassLocation& OutClassLocation, FText* const OutFailReason = nullptr);
@@ -217,7 +217,7 @@ public:
 	/**
 	 * Update the list of supported target platforms based upon the parameters provided
 	 * This will take care of checking out and saving the updated .uproject file automatically
-	 * 
+	 *
 	 * @param	InPlatformName		Name of the platform to target (eg, WindowsNoEditor)
 	 * @param	bIsSupported		true if the platform should be supported by this project, false if it should not
 	 */
@@ -249,7 +249,7 @@ public:
 
 	/** Returns the contents of the specified template file */
 	static bool ReadTemplateFile(const FString& TemplateFileName, FString& OutFileContents, FText& OutFailReason);
-	
+
 	/** Writes an output file. OutputFilename includes a path */
 	static bool WriteOutputFile(const FString& OutputFilename, const FString& OutputFileContents, FText& OutFailReason);
 
@@ -263,7 +263,7 @@ public:
 	static FString MakeIncludeList(const TArray<FString>& InList);
 
 	/** Returns true if the currently loaded project requires a code build */
-	static bool ProjectRequiresBuild(const FName InPlatformInfoName);		
+	static bool ProjectRequiresBuild(const FName InPlatformInfoName);
 
 	/** Deletes the specified list of files that were created during file creation */
 	static void DeleteCreatedFiles(const FString& RootFolder, const TArray<FString>& CreatedFiles);
@@ -293,7 +293,7 @@ private:
 	/** Insert any required feature packs into the DefaultGame.ini file */
 	static bool InsertFeaturePacksIntoINIFile(const FProjectInformation& InProjectInfo, FText& OutFailReason);
 
-	/* 
+	/*
 	 * Insert the addition files from any feature packs specified in the temapalte defs file
 	 * @param InProjectInfo		Project infor to add content for
 	 * @param CreatedFiles		List of files we copied
@@ -324,8 +324,8 @@ private:
 	/** Creates ini files for a new project. On failure, OutFailReason will be populated. */
 	static bool GenerateConfigFiles(const FProjectInformation& InProjectInfo, TArray<FString>& OutCreatedFiles, FText& OutFailReason);
 
-	/* Creates new ini files for a specific project. This is used for turning on the AudioMixer by default on specific platforms in GenerateConfigFiles. */
-	static bool GeneratePlatformConfigFiles(const FProjectInformation& InProjectInfo, const FString& InPlatformName, FText& OutFailReason);
+	/* Creates new ini files for a specific project's platform configurations. */
+	static bool GeneratePlatformConfigFiles(const FProjectInformation& InProjectInfo, FText& OutFailReason);
 
 	/** Creates the basic source code for a new project. On failure, OutFailReason will be populated. */
 	static bool GenerateBasicSourceCode(const FString& NewProjectSourcePath, const FString& NewProjectName, const FString& NewProjectRoot, TArray<FString>& OutGeneratedStartupModuleNames, TArray<FString>& OutCreatedFiles, FText& OutFailReason);
@@ -518,7 +518,4 @@ private:
 
 	// Whether we should use AudioMixer for all platforms:
 	static bool bUseAudioMixerForAllPlatforms;
-	
-	// List of platforms we want to explicitly enable the new audio engine for when generating a new project.
-	static TArray<FString> AudioMixerEnabledPlatforms;
 };
diff --git a/Engine/Source/Runtime/Engine/Private/AudioPluginUtilities.cpp b/Engine/Source/Runtime/Engine/Private/AudioPluginUtilities.cpp
--- a/Engine/Source/Runtime/Engine/Private/AudioPluginUtilities.cpp
+++ b/Engine/Source/Runtime/Engine/Private/AudioPluginUtilities.cpp
@@ -6,74 +6,111 @@
 #include "CoreGlobals.h"
 
 
+/** Get the target setting name for each platform type. */
+FORCEINLINE const TCHAR* GetPluginConfigName(EAudioPlugin PluginType)
+{
+	switch (PluginType)
+	{
+		case EAudioPlugin::SPATIALIZATION:
+			return TEXT("SpatializationPlugin");
+
+		case EAudioPlugin::REVERB:
+			return TEXT("ReverbPlugin");
+
+		case EAudioPlugin::OCCLUSION:
+			return TEXT("OcclusionPlugin");
+
+		default:
+			checkf(false, TEXT("Undefined audio plugin type."));
+			return TEXT("");
+	}
+}
+
+/************************************************************************/
+/* Plugin Utilities                                                     */
+/************************************************************************/
 /** Platform config section for each platform's target settings. */
-FORCEINLINE const TCHAR* GetPlatformConfigSection(EAudioPlatform AudioPlatform)
+const TCHAR* AudioPluginUtilities::GetPlatformConfigSection(EAudioPlatform AudioPlatform)
 {
+	static const FString UnknownConfig = TEXT("");
+
 	switch (AudioPlatform)
 	{
 		case EAudioPlatform::Windows:
-			return TEXT("/Script/WindowsTargetPlatform.WindowsTargetSettings");
+		{
+			static const FString WindowsConfig = TEXT("/Script/WindowsTargetPlatform.WindowsTargetSettings");
+			return *WindowsConfig;
+		}
 
 		case EAudioPlatform::Mac:
-			return TEXT("/Script/MacTargetPlatform.MacTargetSettings");
+		{
+			static const FString MacConfig = TEXT("/Script/MacTargetPlatform.MacTargetSettings");
+			return *MacConfig;
+		}
 
 		case EAudioPlatform::Linux:
-			return TEXT("/Script/LinuxTargetPlatform.LinuxTargetSettings");
+		{
+			static const FString LinuxConfig = TEXT("/Script/LinuxTargetPlatform.LinuxTargetSettings");
+			return *LinuxConfig;
+		}
 
 		case EAudioPlatform::IOS:
-			return TEXT("/Script/IOSRuntimeSettings.IOSRuntimeSettings");
+		{
+			static const FString IOSConfig = TEXT("/Script/IOSRuntimeSettings.IOSRuntimeSettings");
+			return *IOSConfig;
+		}
 
 		case EAudioPlatform::Android:
-			return TEXT("/Script/AndroidRuntimeSettings.AndroidRuntimeSettings");
+		{
+			static const FString AndroidConfig = TEXT("/Script/AndroidRuntimeSettings.AndroidRuntimeSettings");
+			return *AndroidConfig;
+		}
 
 		case EAudioPlatform::XboxOne:
-			return TEXT("/Script/XboxOnePlatformEditor.XboxOneTargetSettings");
+		{
+			static const FString XBoxConfig = TEXT("/Script/XboxOnePlatformEditor.XboxOneTargetSettings");
+			return *XBoxConfig;
+		}
 
 		case EAudioPlatform::Playstation4:
-			return TEXT("/Script/PS4PlatformEditor.PS4TargetSettings");
+		{
+			static const FString Playstation4Config = TEXT("/Script/PS4PlatformEditor.PS4TargetSettings");
+			return *Playstation4Config;
+		}
 
 		case EAudioPlatform::Switch:
-			return TEXT("/Script/SwitchRuntimeSettings.SwitchRuntimeSettings");
+		{
+			static const FString SwitchConfig = TEXT("/Script/SwitchRuntimeSettings.SwitchRuntimeSettings");
+			return *SwitchConfig;
+		}
 
 		case EAudioPlatform::HTML5:
-			return TEXT("/Script/HTML5PlatformEditor.HTML5TargetSettings");
+		{
+			static const FString HTML5Config = TEXT("/Script/HTML5PlatformEditor.HTML5TargetSettings");
+			return *HTML5Config;
+		}
 
 		case EAudioPlatform::Lumin:
-			return TEXT("/Script/LuminRuntimeSettings.LuminRuntimeSettings");
+		{
+			static FString LuminConfig = TEXT("/Script/LuminRuntimeSettings.LuminRuntimeSettings");
+			return *LuminConfig;
+		}
 
 		case EAudioPlatform::Unknown:
-			return TEXT("");
+		{
+			return *UnknownConfig;
+		}
 
 		default:
+		{
 			checkf(false, TEXT("Undefined audio platform."));
 			break;
+		}
 	}
-	return TEXT("");
-}
 
-/** Get the target setting name for each platform type. */
-FORCEINLINE const TCHAR* GetPluginConfigName(EAudioPlugin PluginType)
-{
-	switch (PluginType)
-	{
-		case EAudioPlugin::SPATIALIZATION:
-			return TEXT("SpatializationPlugin");
-
-		case EAudioPlugin::REVERB:
-			return TEXT("ReverbPlugin");
-
-		case EAudioPlugin::OCCLUSION:
-			return TEXT("OcclusionPlugin");
-
-		default:
-			checkf(false, TEXT("Undefined audio plugin type."));
-			return TEXT("");
-	}
+	return *UnknownConfig;
 }
 
-/************************************************************************/
-/* Plugin Utilities                                                     */
-/************************************************************************/
 IAudioSpatializationFactory* AudioPluginUtilities::GetDesiredSpatializationPlugin(EAudioPlatform AudioPlatform)
 {
 	//Get the name of the desired spatialization plugin:
diff --git a/Engine/Source/Runtime/Engine/Public/AudioPluginUtilities.h b/Engine/Source/Runtime/Engine/Public/AudioPluginUtilities.h
--- a/Engine/Source/Runtime/Engine/Public/AudioPluginUtilities.h
+++ b/Engine/Source/Runtime/Engine/Public/AudioPluginUtilities.h
@@ -33,11 +33,14 @@ struct ENGINE_API AudioPluginUtilities
 	static const EAudioPlatform CurrentPlatform = EAudioPlatform::Unknown;
 #endif
 
+	/** Platform config section for each platform's target settings. */
+	static const TCHAR* GetPlatformConfigSection(EAudioPlatform AudioPlatform);
+
 	/*
 	 * These functions return a pointer to the plugin factory
 	 * that matches the plugin name specified in the target
 	 * platform's settings.
-	 * 
+	 *
 	 * if no matching plugin is found, nullptr is returned.
 	 */
 	static IAudioSpatializationFactory* GetDesiredSpatializationPlugin(EAudioPlatform AudioPlatform);
