Fixed InsertFeaturePacksIntoINIFile() function creating an invalid entry in the DefaultGame.ini file

diff --git a/Engine/Source/Editor/GameProjectGeneration/Private/GameProjectUtils.cpp b/Engine/Source/Editor/GameProjectGeneration/Private/GameProjectUtils.cpp
--- a/Engine/Source/Editor/GameProjectGeneration/Private/GameProjectUtils.cpp
+++ b/Engine/Source/Editor/GameProjectGeneration/Private/GameProjectUtils.cpp
@@ -3922,11 +3922,11 @@ bool GameProjectUtils::InsertFeaturePacksIntoINIFile(const FProjectInformation&
 		FString StarterPack;
 		if (InProjectInfo.TargetedHardware == EHardwareClass::Mobile)
 		{
-			StarterPack = TEXT("InsertPack=(PackSource=\"MobileStarterContent") + DefaultFeaturePackExtension + TEXT(",PackName=\"StarterContent\")");
+			StarterPack = TEXT("InsertPack=(PackSource=\"MobileStarterContent") + DefaultFeaturePackExtension + TEXT("\",PackName=\"StarterContent\")");
 		}
 		else
 		{
-			StarterPack = TEXT("InsertPack=(PackSource=\"StarterContent")  + DefaultFeaturePackExtension + TEXT(",PackName=\"StarterContent\")");
+			StarterPack = TEXT("InsertPack=(PackSource=\"StarterContent")  + DefaultFeaturePackExtension + TEXT("\",PackName=\"StarterContent\")");
 		}
 		PackList.Add(StarterPack);
 	}
