Fixed issue where new line character wasn't being inserted after PCH when generating new C++ classes

diff --git a/Engine/Source/Editor/GameProjectGeneration/Private/GameProjectUtils.cpp b/Engine/Source/Editor/GameProjectGeneration/Private/GameProjectUtils.cpp
--- a/Engine/Source/Editor/GameProjectGeneration/Private/GameProjectUtils.cpp
+++ b/Engine/Source/Editor/GameProjectGeneration/Private/GameProjectUtils.cpp
@@ -3031,11 +3031,11 @@ bool GameProjectUtils::GenerateClassCPPFile(const FString& NewCPPFileName, const
 	FinalOutput = FinalOutput.Replace(TEXT("%UNPREFIXED_CLASS_NAME%"), *UnPrefixedClassName, ESearchCase::CaseSensitive);
 	FinalOutput = FinalOutput.Replace(TEXT("%PREFIXED_CLASS_NAME%"), *PrefixedClassName, ESearchCase::CaseSensitive);
 	FinalOutput = FinalOutput.Replace(TEXT("%MODULE_NAME%"), *ModuleInfo.ModuleName, ESearchCase::CaseSensitive);
+	FinalOutput = FinalOutput.Replace(TEXT("%PCH_INCLUDE_DIRECTIVE%"), *PchIncludeDirective, ESearchCase::CaseSensitive);
 
 	// Special case where where the wildcard ends with a new line
 	const bool bLeadingTab = false;
 	const bool bTrailingNewLine = true;
-	FinalOutput = ReplaceWildcard(FinalOutput, TEXT("%PCH_INCLUDE_DIRECTIVE%"), *PchIncludeDirective, bLeadingTab, bTrailingNewLine);
 	FinalOutput = ReplaceWildcard(FinalOutput, TEXT("%ADDITIONAL_INCLUDE_DIRECTIVES%"), *AdditionalIncludesStr, bLeadingTab, bTrailingNewLine);
 	FinalOutput = ReplaceWildcard(FinalOutput, TEXT("%EVENTUAL_CONSTRUCTOR_DEFINITION%"), *EventualConstructorDefinition, bLeadingTab, bTrailingNewLine);
 	FinalOutput = ReplaceWildcard(FinalOutput, TEXT("%ADDITIONAL_MEMBER_DEFINITIONS%"), *AdditionalMemberDefinitions, bLeadingTab, bTrailingNewLine);
