[NO NON-PII COMMIT MESSAGE]

diff --git a/Engine/Source/Editor/GameProjectGeneration/Private/GameProjectUtils.cpp b/Engine/Source/Editor/GameProjectGeneration/Private/GameProjectUtils.cpp
--- a/Engine/Source/Editor/GameProjectGeneration/Private/GameProjectUtils.cpp
+++ b/Engine/Source/Editor/GameProjectGeneration/Private/GameProjectUtils.cpp
@@ -88,6 +88,16 @@ TWeakPtr<SNotificationItem> GameProjectUtils::WarningProjectNameNotification = N
 
 FString GameProjectUtils::DefaultFeaturePackExtension(TEXT(".upack"));	
 
+bool GameProjectUtils::bUseAudioMixerForAllPlatforms = true;
+
+TArray<FString> GameProjectUtils::AudioMixerEnabledPlatforms(
+{
+	// If bUseAudioMixerForAllPlatforms is set to false,
+	// This can be used to only flag specific platforms to use the new audio engine on new projects.
+	// For example, for windows:
+	//TEXT("Windows")
+});
+
 FText FNewClassInfo::GetClassName() const
 {
 	switch(ClassType)
@@ -1814,6 +1824,13 @@ bool GameProjectUtils::GenerateConfigFiles(const FProjectInformation& InProjectI
 		FileContents += GetHardwareConfigString(InProjectInfo);
 		FileContents += LINE_TERMINATOR;
 		
+		if(bUseAudioMixerForAllPlatforms)
+		{
+			FileContents += TEXT("[Audio]") LINE_TERMINATOR;
+			FileContents += TEXT("UseAudioMixer=True") LINE_TERMINATOR;
+			FileContents += LINE_TERMINATOR;
+		}
+		
 		if (InProjectInfo.bCopyStarterContent)
 		{
 			FString SpecificEditorStartupMap;
@@ -1886,9 +1903,42 @@ bool GameProjectUtils::GenerateConfigFiles(const FProjectInformation& InProjectI
 		}
 	}
 
+	// inis for any audiomixer-enabled platforms:
+	{
+		for (const FString& PlatformName : AudioMixerEnabledPlatforms)
+		{
+			if (!GeneratePlatformConfigFiles(InProjectInfo, PlatformName, OutFailReason))
+			{
+				return false;
+			}
+		}
+	}
+
 	return true;
 }
 
+bool GameProjectUtils::GeneratePlatformConfigFiles(const FProjectInformation& InProjectInfo, const FString& InPlatformName, FText& OutFailReason)
+{
+	const FString NewProjectFolder = FPaths::GetPath(InProjectInfo.ProjectFilename);
+
+	FString ProjectConfigPath = NewProjectFolder / TEXT("Config");
+
+	const FString PlatformEngineIniFilename = ProjectConfigPath / InPlatformName / InPlatformName + TEXT("Engine.ini");
+	FString FileContents;
+
+	FileContents += TEXT("[Audio]") LINE_TERMINATOR;
+	FileContents += TEXT("UseAudioMixer=True");
+
+	if (WriteOutputFile(PlatformEngineIniFilename, FileContents, OutFailReason))
+	{
+		return true;
+	}
+	else
+	{
+		return false;
+	}
+}
+
 bool GameProjectUtils::GenerateBasicSourceCode(TArray<FString>& OutCreatedFiles, FText& OutFailReason)
 {
 	TArray<FString> StartupModuleNames;
diff --git a/Engine/Source/Editor/GameProjectGeneration/Public/GameProjectUtils.h b/Engine/Source/Editor/GameProjectGeneration/Public/GameProjectUtils.h
--- a/Engine/Source/Editor/GameProjectGeneration/Public/GameProjectUtils.h
+++ b/Engine/Source/Editor/GameProjectGeneration/Public/GameProjectUtils.h
@@ -322,6 +322,9 @@ private:
 	/** Creates ini files for a new project. On failure, OutFailReason will be populated. */
 	static bool GenerateConfigFiles(const FProjectInformation& InProjectInfo, TArray<FString>& OutCreatedFiles, FText& OutFailReason);
 
+	/* Creates new ini files for a specific project. This is used for turning on the AudioMixer by default on specific platforms in GenerateConfigFiles. */
+	static bool GeneratePlatformConfigFiles(const FProjectInformation& InProjectInfo, const FString& InPlatformName, FText& OutFailReason);
+
 	/** Creates the basic source code for a new project. On failure, OutFailReason will be populated. */
 	static bool GenerateBasicSourceCode(const FString& NewProjectSourcePath, const FString& NewProjectName, const FString& NewProjectRoot, TArray<FString>& OutGeneratedStartupModuleNames, TArray<FString>& OutCreatedFiles, FText& OutFailReason);
 
@@ -509,5 +512,11 @@ private:
 
 	static TWeakPtr<SNotificationItem> UpdateGameProjectNotification;
 	static TWeakPtr<SNotificationItem> WarningProjectNameNotification;
-	static FString DefaultFeaturePackExtension;	
+	static FString DefaultFeaturePackExtension;
+
+	// Whether we should use AudioMixer for all platforms:
+	static bool bUseAudioMixerForAllPlatforms;
+	
+	// List of platforms we want to explicitly enable the new audio engine for when generating a new project.
+	static TArray<FString> AudioMixerEnabledPlatforms;
 };
