[NO NON-PII COMMIT MESSAGE]

diff --git a/Engine/Source/Editor/GameProjectGeneration/Private/GameProjectUtils.cpp b/Engine/Source/Editor/GameProjectGeneration/Private/GameProjectUtils.cpp
--- a/Engine/Source/Editor/GameProjectGeneration/Private/GameProjectUtils.cpp
+++ b/Engine/Source/Editor/GameProjectGeneration/Private/GameProjectUtils.cpp
@@ -88,7 +88,7 @@ TWeakPtr<SNotificationItem> GameProjectUtils::WarningProjectNameNotification = N
 
 FString GameProjectUtils::DefaultFeaturePackExtension(TEXT(".upack"));	
 
-bool GameProjectUtils::bUseAudioMixerForAllPlatforms = true;
+bool GameProjectUtils::bUseAudioMixerForAllPlatforms = false;
 
 TArray<FString> GameProjectUtils::AudioMixerEnabledPlatforms(
 {
