Fix issues related to hot-reload and LiveCoding co-existing.

* Removed code to invalidate makefiles when adding new source files. UBT should be reliable enough to make this determination itself nowadays, and ignored -invalidatemakefilesonly argument was causing modules to be recompiled.
* Fixed incorrect config section name when determining whether to allow hot reload from IDE. Now prevents hot reload from IDE when live coding is enabled.
* Added error message when trying to add a new class with Live Coding enabled.
* Added error messages when trying to start Live Coding after a hot reload has taken place.
* Added error messages when trying to hot reload with Live Coding enabled.

diff --git a/Engine/Source/Developer/DesktopPlatform/Private/DesktopPlatformBase.cpp b/Engine/Source/Developer/DesktopPlatform/Private/DesktopPlatformBase.cpp
--- a/Engine/Source/Developer/DesktopPlatform/Private/DesktopPlatformBase.cpp
+++ b/Engine/Source/Developer/DesktopPlatform/Private/DesktopPlatformBase.cpp
@@ -637,40 +637,6 @@ bool FDesktopPlatformBase::GenerateProjectFiles(const FString& RootDir, const FS
 	return bRes;
 }
 
-bool FDesktopPlatformBase::InvalidateMakefiles(const FString& RootDir, const FString& ProjectFileName, FFeedbackContext* Warn)
-{
-	// Composes the platform, and config (eg, "Win64 Development")
-	FString Arguments = FString::Printf(TEXT("%s %s"), FPlatformMisc::GetUBTPlatform(), FModuleManager::GetUBTConfiguration());
-
-	// -TargetType=Editor tells UBT to work out the editor target name from the project we provided
-	Arguments += TEXT(" -TargetType=Editor");
-
-	// Add the project path
-	if ( !ProjectFileName.IsEmpty() )
-	{
-		Arguments += FString::Printf(TEXT(" -Project=\"%s\""), *IFileManager::Get().ConvertToAbsolutePathForExternalAppForRead(*ProjectFileName));
-	}
-	
-	// -invalidatemakefilesonly tells UBT to invalidate its UBT makefiles without building
-	Arguments += TEXT(" -invalidatemakefilesonly");
-
-	// Compile UnrealBuildTool if it doesn't exist. This can happen if we're just copying source from somewhere.
-	bool bRes = true;
-	Warn->BeginSlowTask(LOCTEXT("InvalidateMakefiles", "Invalidating makefiles..."), true, true);
-	if(!FPaths::FileExists(GetUnrealBuildToolExecutableFilename(RootDir)))
-	{
-		Warn->StatusUpdate(0, 1, LOCTEXT("BuildingUBT", "Building UnrealBuildTool..."));
-		bRes = BuildUnrealBuildTool(RootDir, *Warn);
-	}
-	if(bRes)
-	{
-		Warn->StatusUpdate(0, 1, LOCTEXT("InvalidateMakefiles", "Invalidating makefiles..."));
-		bRes = RunUnrealBuildTool(LOCTEXT("InvalidateMakefiles", "Invalidating makefiles..."), RootDir, Arguments, Warn);
-	}
-	Warn->EndSlowTask();
-	return bRes;
-}
-
 bool FDesktopPlatformBase::IsUnrealBuildToolAvailable()
 {
 	// If using installed build and the unreal build tool executable exists, then UBT is available. Otherwise check it can be built.
diff --git a/Engine/Source/Developer/DesktopPlatform/Private/DesktopPlatformBase.h b/Engine/Source/Developer/DesktopPlatform/Private/DesktopPlatformBase.h
--- a/Engine/Source/Developer/DesktopPlatform/Private/DesktopPlatformBase.h
+++ b/Engine/Source/Developer/DesktopPlatform/Private/DesktopPlatformBase.h
@@ -45,7 +45,6 @@ public:
 	virtual bool CleanGameProject(const FString& ProjectDir, FString& OutFailPath, FFeedbackContext* Warn) override;
 	virtual bool CompileGameProject(const FString& RootDir, const FString& ProjectFileName, FFeedbackContext* Warn) override;
 	virtual bool GenerateProjectFiles(const FString& RootDir, const FString& ProjectFileName, FFeedbackContext* Warn, FString LogFilePath = FString()) override;
-	virtual bool InvalidateMakefiles(const FString& RootDir, const FString& ProjectFileName, FFeedbackContext* Warn) override;
 	virtual bool IsUnrealBuildToolAvailable() override;
 	virtual bool InvokeUnrealBuildToolSync(const FString& InCmdLineParams, FOutputDevice &Ar, bool bSkipBuildUBT, int32& OutReturnCode, FString& OutProcOutput) override;
 	virtual FProcHandle InvokeUnrealBuildToolAsync(const FString& InCmdLineParams, FOutputDevice &Ar, void*& OutReadPipe, void*& OutWritePipe, bool bSkipBuildUBT = false) override;
diff --git a/Engine/Source/Developer/DesktopPlatform/Public/IDesktopPlatform.h b/Engine/Source/Developer/DesktopPlatform/Public/IDesktopPlatform.h
--- a/Engine/Source/Developer/DesktopPlatform/Public/IDesktopPlatform.h
+++ b/Engine/Source/Developer/DesktopPlatform/Public/IDesktopPlatform.h
@@ -348,16 +348,6 @@ public:
 	*/
 	virtual bool GenerateProjectFiles(const FString& RootDir, const FString& ProjectFileName, FFeedbackContext* Warn, FString LogFilePath = FString()) = 0;
 
-	/**
-	* Invalidate makefiles for project (to UBT regenerate them at startup).
-	*
-	* @param RootDir			Engine root directory for the project to use.
-	* @param ProjectFileName	Filename of the project to update
-	* @param Warn				Feedback context to use for progress updates
-	* @return true if project files were generated successfully.
-	*/
-	virtual bool InvalidateMakefiles(const FString& RootDir, const FString& ProjectFileName, FFeedbackContext* Warn) = 0;
-
 	/**
 	* Determines whether UnrealBuildTool is available
 	*
diff --git a/Engine/Source/Developer/HotReload/HotReload.Build.cs b/Engine/Source/Developer/HotReload/HotReload.Build.cs
--- a/Engine/Source/Developer/HotReload/HotReload.Build.cs
+++ b/Engine/Source/Developer/HotReload/HotReload.Build.cs
@@ -5,33 +5,38 @@ public class HotReload : ModuleRules
 {
 	public HotReload(ReadOnlyTargetRules Target) : base(Target)
 	{
-			PublicDependencyModuleNames.AddRange(
-				new string[]
-				{
-					"Core",
-					"CoreUObject",
-				}
-			);
+		PublicDependencyModuleNames.AddRange(
+			new string[]
+			{
+				"Core",
+				"CoreUObject",
+			}
+		);
+
+		PrivateDependencyModuleNames.AddRange(
+			new string[]
+			{
+				"Analytics",
+				"DirectoryWatcher",
+				"DesktopPlatform",
+				"Projects"
+			}
+		);
 
-			PrivateDependencyModuleNames.AddRange(
-				new string[]
+        if (Target.bCompileAgainstEngine)
+        {
+            PrivateDependencyModuleNames.AddRange(
+				new string[] 
 				{
-					"Analytics",
-					"DirectoryWatcher",
-					"DesktopPlatform",
-					"Projects"
+					"Engine",
+					"UnrealEd", 
 				}
 			);
+        }
 
-            if (Target.bCompileAgainstEngine)
-            {
-                PrivateDependencyModuleNames.AddRange(
-				    new string[] 
-				    {
-					    "Engine",
-					    "UnrealEd", 
-				    }
-			    );
-            }
+		if(Target.bWithLiveCoding)
+		{
+			PrivateIncludePathModuleNames.Add("LiveCoding");
+		}
 	}
 }
diff --git a/Engine/Source/Developer/HotReload/Private/HotReload.cpp b/Engine/Source/Developer/HotReload/Private/HotReload.cpp
--- a/Engine/Source/Developer/HotReload/Private/HotReload.cpp
+++ b/Engine/Source/Developer/HotReload/Private/HotReload.cpp
@@ -42,6 +42,10 @@
 #include "Misc/ScopeExit.h"
 #include "Algo/Transform.h"
 
+#if WITH_LIVE_CODING
+#include "ILiveCodingModule.h"
+#endif
+
 #if WITH_EDITOR
 #include "Editor.h"
 #endif
@@ -601,6 +605,16 @@ FString FHotReloadModule::GetModuleCompileMethod(FName InModuleName)
 bool FHotReloadModule::RecompileModule(const FName InModuleName, const bool bReloadAfterRecompile, FOutputDevice &Ar, bool bFailIfGeneratedCodeChanges, bool bForceCodeProject)
 {
 #if WITH_HOT_RELOAD
+
+#if WITH_LIVE_CODING
+	ILiveCodingModule* LiveCoding = FModuleManager::GetModulePtr<ILiveCodingModule>(LIVE_CODING_MODULE_NAME);
+	if (LiveCoding != nullptr && LiveCoding->IsEnabledForSession())
+	{
+		UE_LOG(LogHotReload, Error, TEXT("Unable to hot-reload modules while Live Coding is enabled."));
+		return false;
+	}
+#endif
+
 	UE_LOG(LogHotReload, Log, TEXT("Recompiling module %s..."), *InModuleName.ToString());
 
 	// This is an internal request for hot-reload (not from IDE)
@@ -1423,6 +1437,15 @@ bool FHotReloadModule::Tick(float DeltaTime)
 		return true;
 	}
 
+	// Early out if live coding is enabled
+#if WITH_LIVE_CODING
+	ILiveCodingModule* LiveCoding = FModuleManager::GetModulePtr<ILiveCodingModule>(LIVE_CODING_MODULE_NAME);
+	if (LiveCoding != nullptr && LiveCoding->IsEnabledForSession())
+	{
+        return false;
+	}
+#endif
+
 #if WITH_EDITOR
 	if (GEditor)
 	{
diff --git a/Engine/Source/Developer/Windows/LiveCoding/Private/LiveCodingModule.cpp b/Engine/Source/Developer/Windows/LiveCoding/Private/LiveCodingModule.cpp
--- a/Engine/Source/Developer/Windows/LiveCoding/Private/LiveCodingModule.cpp
+++ b/Engine/Source/Developer/Windows/LiveCoding/Private/LiveCodingModule.cpp
@@ -137,9 +137,9 @@ void FLiveCodingModule::EnableForSession(bool bEnable)
 			UE_LOG(LogLiveCoding, Display, TEXT("Console will be hidden but remain running in the background. Restart to disable completely."));
 			LppSetActive(false);
 			LppSetVisible(false);
+			bEnabledForSession = false;
 		}
 	}
-	bEnabledForSession = bEnable;
 }
 
 bool FLiveCodingModule::IsEnabledForSession() const
@@ -147,6 +147,18 @@ bool FLiveCodingModule::IsEnabledForSession() const
 	return bEnabledForSession;
 }
 
+bool FLiveCodingModule::CanEnableForSession() const
+{
+#if !IS_MONOLITHIC
+	FModuleManager& ModuleManager = FModuleManager::Get();
+	if(ModuleManager.HasAnyOverridenModuleFilename())
+	{
+		return false;
+	}
+#endif
+	return true;
+}
+
 bool FLiveCodingModule::HasStarted() const
 {
 	return bStarted;
@@ -193,6 +205,13 @@ bool FLiveCodingModule::StartLiveCoding()
 {
 	if(!bStarted)
 	{
+		// Make sure there aren't any hot reload modules already active
+		if (!CanEnableForSession())
+		{
+			UE_LOG(LogLiveCoding, Error, TEXT("Unable to start live coding session. Some modules have already been hot reloaded."));
+			return false;
+		}
+
 		// Setup the console path
 		GLiveCodingConsolePath = ConsolePathVariable->GetString();
 		if (!FPaths::FileExists(GLiveCodingConsolePath))
diff --git a/Engine/Source/Developer/Windows/LiveCoding/Private/LiveCodingModule.h b/Engine/Source/Developer/Windows/LiveCoding/Private/LiveCodingModule.h
--- a/Engine/Source/Developer/Windows/LiveCoding/Private/LiveCodingModule.h
+++ b/Engine/Source/Developer/Windows/LiveCoding/Private/LiveCodingModule.h
@@ -26,6 +26,7 @@ public:
 	virtual bool IsEnabledByDefault() const override;
 	virtual void EnableForSession(bool bInEnabled) override;
 	virtual bool IsEnabledForSession() const override;
+	virtual bool CanEnableForSession() const override;
 	virtual bool HasStarted() const override;
 	virtual void ShowConsole() override;
 	virtual void Compile() override;
diff --git a/Engine/Source/Developer/Windows/LiveCoding/Private/LiveCodingSettings.h b/Engine/Source/Developer/Windows/LiveCoding/Private/LiveCodingSettings.h
--- a/Engine/Source/Developer/Windows/LiveCoding/Private/LiveCodingSettings.h
+++ b/Engine/Source/Developer/Windows/LiveCoding/Private/LiveCodingSettings.h
@@ -20,7 +20,7 @@ class ULiveCodingSettings : public UObject
     GENERATED_BODY()
 
 public:
-    UPROPERTY(config, EditAnywhere, Category=General, Meta=(DisplayName="Enable Live Coding"))
+    UPROPERTY(config, EditAnywhere, Category=General, Meta=(ConfigRestartRequired=true, DisplayName="Enable Live Coding"))
     bool bEnabled;
 
     UPROPERTY(config, EditAnywhere, Category=General, Meta=(ConfigRestartRequired=true, EditCondition="bEnabled"))
diff --git a/Engine/Source/Developer/Windows/LiveCoding/Public/ILiveCodingModule.h b/Engine/Source/Developer/Windows/LiveCoding/Public/ILiveCodingModule.h
--- a/Engine/Source/Developer/Windows/LiveCoding/Public/ILiveCodingModule.h
+++ b/Engine/Source/Developer/Windows/LiveCoding/Public/ILiveCodingModule.h
@@ -15,6 +15,7 @@ public:
 
 	virtual void EnableForSession(bool bEnabled) = 0;
 	virtual bool IsEnabledForSession() const = 0;
+	virtual bool CanEnableForSession() const = 0;
 
 	virtual bool HasStarted() const = 0;
 
diff --git a/Engine/Source/Editor/GameProjectGeneration/GameProjectGeneration.Build.cs b/Engine/Source/Editor/GameProjectGeneration/GameProjectGeneration.Build.cs
--- a/Engine/Source/Editor/GameProjectGeneration/GameProjectGeneration.Build.cs
+++ b/Engine/Source/Editor/GameProjectGeneration/GameProjectGeneration.Build.cs
@@ -59,5 +59,10 @@ public class GameProjectGeneration : ModuleRules
                 "MainFrame",
             }
 		);
+
+		if(Target.bWithLiveCoding)
+		{
+			PrivateIncludePathModuleNames.Add("LiveCoding");
+		}
 	}
 }
diff --git a/Engine/Source/Editor/GameProjectGeneration/Private/GameProjectUtils.cpp b/Engine/Source/Editor/GameProjectGeneration/Private/GameProjectUtils.cpp
--- a/Engine/Source/Editor/GameProjectGeneration/Private/GameProjectUtils.cpp
+++ b/Engine/Source/Editor/GameProjectGeneration/Private/GameProjectUtils.cpp
@@ -76,6 +76,10 @@
 
 #include "ProjectBuildMutatorFeature.h"
 
+#if WITH_LIVE_CODING
+#include "ILiveCodingModule.h"
+#endif
+
 #define LOCTEXT_NAMESPACE "GameProjectUtils"
 
 #define MAX_PROJECT_PATH_BUFFER_SPACE 130 // Leave a reasonable buffer of additional characters to account for files created in the content directory during or after project generation
@@ -3813,12 +3817,8 @@ GameProjectUtils::EAddCodeToProjectResult GameProjectUtils::AddCodeToProject_Int
 	// First see if we can avoid a full generation by adding the new files to an already open project
 	if ( bProjectHadCodeFiles && FSourceCodeNavigation::AddSourceFiles(CreatedFilesForExternalAppRead) )
 	{
-		// We successfully added the new files to the solution, but we still need to run UBT with -gather to update any UBT makefiles
-		if ( FDesktopPlatformModule::Get()->InvalidateMakefiles(FPaths::RootDir(), FPaths::GetProjectFilePath(), GWarn) )
-		{
-			// We managed the gather, so we can skip running the full generate
-			bGenerateProjectFiles = false;
-		}
+		// We managed the gather, so we can skip running the full generate
+		bGenerateProjectFiles = false;
 	}
 	
 	if ( bGenerateProjectFiles )
@@ -3845,6 +3845,15 @@ GameProjectUtils::EAddCodeToProjectResult GameProjectUtils::AddCodeToProject_Int
 	OutHeaderFilePath = NewHeaderFilename;
 	OutCppFilePath = NewCppFilename;
 
+#if WITH_LIVE_CODING
+	ILiveCodingModule* LiveCoding = FModuleManager::GetModulePtr<ILiveCodingModule>(LIVE_CODING_MODULE_NAME);
+	if (LiveCoding != nullptr && LiveCoding->IsEnabledForSession())
+	{
+		OutFailReason = LOCTEXT("FailedToCompileLiveCodingEnabled", "Adding classes dynamically is not allowed with Live Coding enabled.");
+		return EAddCodeToProjectResult::FailedToHotReload;
+	}
+#endif
+
 	if (!bProjectHadCodeFiles)
 	{
 		// This is the first time we add code to this project so compile its game DLL
diff --git a/Engine/Source/Editor/LevelEditor/Private/LevelEditorActions.cpp b/Engine/Source/Editor/LevelEditor/Private/LevelEditorActions.cpp
--- a/Engine/Source/Editor/LevelEditor/Private/LevelEditorActions.cpp
+++ b/Engine/Source/Editor/LevelEditor/Private/LevelEditorActions.cpp
@@ -1035,7 +1035,15 @@ void FLevelEditorActionCallbacks::RecompileGameCode_Clicked()
 	ILiveCodingModule* LiveCoding = FModuleManager::GetModulePtr<ILiveCodingModule>(LIVE_CODING_MODULE_NAME);
 	if (LiveCoding != nullptr && LiveCoding->IsEnabledByDefault())
 	{
-		LiveCoding->Compile();
+		LiveCoding->EnableForSession(true);
+		if (LiveCoding->IsEnabledForSession())
+		{
+			LiveCoding->Compile();
+		}
+		else
+		{
+			FMessageDialog::Open(EAppMsgType::Ok, LOCTEXT("NoLiveCodingCompileAfterHotReload", "Live Coding cannot be enabled after hot-reload has been used. Please restart the editor."));
+		}
 		return;
 	}
 #endif
@@ -1078,6 +1086,11 @@ void FLevelEditorActionCallbacks::LiveCoding_ToggleEnabled()
 	if (LiveCoding != nullptr)
 	{
 		LiveCoding->EnableByDefault(!LiveCoding->IsEnabledByDefault());
+
+		if (LiveCoding->IsEnabledByDefault() && !LiveCoding->IsEnabledForSession())
+		{
+			FMessageDialog::Open(EAppMsgType::Ok, LOCTEXT("NoEnableLiveCodingAfterHotReload", "Live Coding cannot be enabled after hot-reload has been used. Please restart the editor."));
+		}
 	}
 }
 
@@ -1090,9 +1103,14 @@ bool FLevelEditorActionCallbacks::LiveCoding_IsEnabled( )
 void FLevelEditorActionCallbacks::LiveCoding_StartSession_Clicked()
 {
 	ILiveCodingModule* LiveCoding = FModuleManager::GetModulePtr<ILiveCodingModule>(LIVE_CODING_MODULE_NAME);
-	if (LiveCoding!= nullptr)
+	if (LiveCoding != nullptr)
 	{
 		LiveCoding->EnableForSession(true);
+
+		if (!LiveCoding->IsEnabledForSession())
+		{
+			FMessageDialog::Open(EAppMsgType::Ok, LOCTEXT("NoStartedLiveCodingAfterHotReload", "Live Coding cannot be started after hot-reload has been used. Please restart the editor."));
+		}
 	}
 }
 
diff --git a/Engine/Source/Programs/UnrealBuildTool/System/HotReload.cs b/Engine/Source/Programs/UnrealBuildTool/System/HotReload.cs
--- a/Engine/Source/Programs/UnrealBuildTool/System/HotReload.cs
+++ b/Engine/Source/Programs/UnrealBuildTool/System/HotReload.cs
@@ -141,7 +141,7 @@ namespace UnrealBuildTool
 			// Check if we're using LiveCode instead
 			ConfigHierarchy EditorPerProjectHierarchy = ConfigCache.ReadHierarchy(ConfigHierarchyType.EditorPerProjectUserSettings, DirectoryReference.FromFile(TargetDesc.ProjectFile), TargetDesc.Platform);
 			bool bEnableLiveCode;
-			if(EditorPerProjectHierarchy.GetBool("LiveCode", "Enabled", out bEnableLiveCode) && bEnableLiveCode)
+			if(EditorPerProjectHierarchy.GetBool("/Script/LiveCoding.LiveCodingSettings", "bEnabled", out bEnableLiveCode) && bEnableLiveCode)
 			{
 				return false;
 			}
diff --git a/Engine/Source/Runtime/Core/Private/Modules/ModuleManager.cpp b/Engine/Source/Runtime/Core/Private/Modules/ModuleManager.cpp
--- a/Engine/Source/Runtime/Core/Private/Modules/ModuleManager.cpp
+++ b/Engine/Source/Runtime/Core/Private/Modules/ModuleManager.cpp
@@ -934,6 +934,20 @@ void FModuleManager::SetModuleFilename(FName ModuleName, const FString& Filename
 		Module->OriginalFilename = Filename;
 	}
 }
+
+bool FModuleManager::HasAnyOverridenModuleFilename() const
+{
+	FScopeLock Lock(&ModulesCriticalSection);
+	for (const TPair<FName, ModuleInfoRef>& ModuleIt : Modules)
+	{
+		const FModuleInfo& CurModule = *ModuleIt.Value;
+		if(CurModule.Filename != CurModule.OriginalFilename)
+		{
+			return true;
+		}
+	}
+	return false;
+}
 #endif
 
 void FModuleManager::ResetModulePathsCache()
diff --git a/Engine/Source/Runtime/Core/Public/Modules/ModuleManager.h b/Engine/Source/Runtime/Core/Public/Modules/ModuleManager.h
--- a/Engine/Source/Runtime/Core/Public/Modules/ModuleManager.h
+++ b/Engine/Source/Runtime/Core/Public/Modules/ModuleManager.h
@@ -415,6 +415,9 @@ public:
 
 	/** Sets the filename for a module. The module is not reloaded immediately, but the new name will be used for subsequent unload/load events. */
 	void SetModuleFilename(FName ModuleName, const FString& Filename);
+
+	/** Determines if any non-default module instances are loaded (eg. hot reloaded modules) */
+	bool HasAnyOverridenModuleFilename() const;
 #endif
 
 public:
