GitHub #5776 : Changed empty class prefix to be consistent with UE coding standards

Also updated the prefix for the UObject type.

diff --git a/Engine/Source/Editor/GameProjectGeneration/Private/GameProjectUtils.cpp b/Engine/Source/Editor/GameProjectGeneration/Private/GameProjectUtils.cpp
--- a/Engine/Source/Editor/GameProjectGeneration/Private/GameProjectUtils.cpp
+++ b/Engine/Source/Editor/GameProjectGeneration/Private/GameProjectUtils.cpp
@@ -186,10 +186,10 @@ FString FNewClassInfo::GetClassPrefixCPP() const
 	switch(ClassType)
 	{
 	case EClassType::UObject:
-		return BaseClass ? BaseClass->GetPrefixCPP() : TEXT("");
+		return BaseClass ? BaseClass->GetPrefixCPP() : TEXT("U");
 
 	case EClassType::EmptyCpp:
-		return TEXT("");
+		return TEXT("F");
 
 	case EClassType::SlateWidget:
 		return TEXT("S");
