PR #5520: Fix indentation that didn't match other code in same file. (Contributed by [REDACTED_PERSON])

diff --git a/Engine/Source/Editor/GameProjectGeneration/Private/GameProjectUtils.cpp b/Engine/Source/Editor/GameProjectGeneration/Private/GameProjectUtils.cpp
--- a/Engine/Source/Editor/GameProjectGeneration/Private/GameProjectUtils.cpp
+++ b/Engine/Source/Editor/GameProjectGeneration/Private/GameProjectUtils.cpp
@@ -756,46 +756,46 @@ void GameProjectUtils::CheckForOutOfDateGameProjectFile()
 			if ( ProjectStatus.bRequiresUpdate )
 			{
 				bRequiresUpdate = true;
-				}
+			}
 		}
 
 		// Get the current project descriptor
-			const FProjectDescriptor* Project = IProjectManager::Get().GetCurrentProject();
+		const FProjectDescriptor* Project = IProjectManager::Get().GetCurrentProject();
 
 		// Check if there are any installed plugins that need to be added as a reference
-				TArray<FPluginReferenceDescriptor> NewPluginReferences = Project->Plugins;
-				for(TSharedRef<IPlugin>& Plugin: IPluginManager::Get().GetEnabledPlugins())
-				{
-					if(Plugin->GetDescriptor().bInstalled && Project->FindPluginReferenceIndex(Plugin->GetName()) == INDEX_NONE)
-					{
-						FPluginReferenceDescriptor PluginReference(Plugin->GetName(), true);
-						NewPluginReferences.Add(PluginReference);
+		TArray<FPluginReferenceDescriptor> NewPluginReferences = Project->Plugins;
+		for(TSharedRef<IPlugin>& Plugin: IPluginManager::Get().GetEnabledPlugins())
+		{
+			if(Plugin->GetDescriptor().bInstalled && Project->FindPluginReferenceIndex(Plugin->GetName()) == INDEX_NONE)
+			{
+				FPluginReferenceDescriptor PluginReference(Plugin->GetName(), true);
+				NewPluginReferences.Add(PluginReference);
 				bRequiresUpdate = true;
-					}
-				}
+			}
+		}
 
-				// Check if there are any referenced plugins that do not have a matching supported plugins list
-				for(FPluginReferenceDescriptor& Reference: NewPluginReferences)
+		// Check if there are any referenced plugins that do not have a matching supported plugins list
+		for(FPluginReferenceDescriptor& Reference: NewPluginReferences)
+		{
+			if(Reference.bEnabled)
+			{
+				TSharedPtr<IPlugin> Plugin = IPluginManager::Get().FindPlugin(Reference.Name);
+				if(Plugin.IsValid())
 				{
-					if(Reference.bEnabled)
+					const FPluginDescriptor& Descriptor = Plugin->GetDescriptor();
+					if(Reference.MarketplaceURL != Descriptor.MarketplaceURL)
 					{
-						TSharedPtr<IPlugin> Plugin = IPluginManager::Get().FindPlugin(Reference.Name);
-						if(Plugin.IsValid())
-						{
-							const FPluginDescriptor& Descriptor = Plugin->GetDescriptor();
-							if(Reference.MarketplaceURL != Descriptor.MarketplaceURL)
-							{
-								Reference.MarketplaceURL = Descriptor.MarketplaceURL;
+						Reference.MarketplaceURL = Descriptor.MarketplaceURL;
 						bRequiresUpdate = true;
-							}
-							if(Reference.SupportedTargetPlatforms != Descriptor.SupportedTargetPlatforms)
-							{
-								Reference.SupportedTargetPlatforms = Descriptor.SupportedTargetPlatforms;
+					}
+					if(Reference.SupportedTargetPlatforms != Descriptor.SupportedTargetPlatforms)
+					{
+						Reference.SupportedTargetPlatforms = Descriptor.SupportedTargetPlatforms;
 						bRequiresUpdate = true;
-							}
-						}
 					}
 				}
+			}
+		}
 
 		// If we have updates pending, show the prompt
 		if (bRequiresUpdate)
@@ -819,10 +819,10 @@ void GameProjectUtils::CheckForOutOfDateGameProjectFile()
 			Info.ButtonDetails.Add(FNotificationButtonInfo(UpdateProjectCancelText, FText(), FSimpleDelegate::CreateStatic(&GameProjectUtils::OnUpdateProjectCancel)));
 
 			if (UpdateGameProjectNotification.IsValid())
-				{
+			{
 				UpdateGameProjectNotification.Pin()->ExpireAndFadeout();
 				UpdateGameProjectNotification.Reset();
-				}
+			}
 
 			UpdateGameProjectNotification = FSlateNotificationManager::Get().AddNotification(Info);
 
