Fix issue where projects with disabled plugins will not cause a blueprint project to have a temporary target.cs file generated.

#rnx

diff --git a/Engine/Source/Editor/GameProjectGeneration/Private/GameProjectUtils.cpp b/Engine/Source/Editor/GameProjectGeneration/Private/GameProjectUtils.cpp
--- a/Engine/Source/Editor/GameProjectGeneration/Private/GameProjectUtils.cpp
+++ b/Engine/Source/Editor/GameProjectGeneration/Private/GameProjectUtils.cpp
@@ -3671,7 +3671,7 @@ bool GameProjectUtils::ProjectRequiresBuild(const FName InPlatformInfoName)
 
 	// check to see if any plugins beyond the defaults have been enabled
 	const PlatformInfo::FPlatformInfo* PlatformInfo = PlatformInfo::FindPlatformInfo(InPlatformInfoName);
-	FName PlatformName = (PlatformInfo != nullptr) ? PlatformInfo->VanillaPlatformName : InPlatformInfoName;
+	FName PlatformName = (PlatformInfo != nullptr) ? PlatformInfo->UBTTargetId : InPlatformInfoName;
 	bRequiresBuild |= !HasDefaultPluginSettings(PlatformName.ToString());
 
 	// check to see if Blueprint nativization is enabled in the Project settings
