Integrated all fixes made to the New Project dialog made in Dev-Editor to Release.

diff --git a/Engine/Source/Editor/GameProjectGeneration/Classes/TemplateProjectDefs.h b/Engine/Source/Editor/GameProjectGeneration/Classes/TemplateProjectDefs.h
--- a/Engine/Source/Editor/GameProjectGeneration/Classes/TemplateProjectDefs.h
+++ b/Engine/Source/Editor/GameProjectGeneration/Classes/TemplateProjectDefs.h
@@ -125,6 +125,10 @@ public:
 	UPROPERTY(config)
 	bool bIsEnterprise;
 
+	/** Is this a blank template? Determines whether we can override the default map when creating the project. */
+	UPROPERTY(config)
+	bool bIsBlank;
+
 	/* Optional list of settings to hide. If none are specified, then all settings are shown. */
 	UPROPERTY(config)
 	TArray<ETemplateSetting> HiddenSettings;
diff --git a/Engine/Source/Editor/GameProjectGeneration/Private/GameProjectUtils.cpp b/Engine/Source/Editor/GameProjectGeneration/Private/GameProjectUtils.cpp
--- a/Engine/Source/Editor/GameProjectGeneration/Private/GameProjectUtils.cpp
+++ b/Engine/Source/Editor/GameProjectGeneration/Private/GameProjectUtils.cpp
@@ -93,8 +93,6 @@
 TWeakPtr<SNotificationItem> GameProjectUtils::UpdateGameProjectNotification = NULL;
 TWeakPtr<SNotificationItem> GameProjectUtils::WarningProjectNameNotification = NULL;
 
-FString GameProjectUtils::DefaultFeaturePackExtension(TEXT(".upack"));
-
 bool GameProjectUtils::bUseAudioMixerForAllPlatforms = false;
 
 struct FAudioDefaultPlatformSettings
@@ -202,6 +200,32 @@ namespace
 
 		return MoveTemp(RaytracingConfig);
 	}
+
+	FString GetDefaultMapConfigString(const FProjectInformation& InProjectInfo)
+	{
+		FString DefaultMapConfig;
+
+		if (InProjectInfo.bIsBlankTemplate &&
+			InProjectInfo.bCopyStarterContent &&
+			GameProjectUtils::IsStarterContentAvailableForNewProjects())
+		{
+			DefaultMapConfig += LINE_TERMINATOR;
+			DefaultMapConfig += TEXT("[/Script/EngineSettings.GameMapsSettings]") LINE_TERMINATOR;
+
+			if (InProjectInfo.TargetedHardware == EHardwareClass::Mobile)
+			{
+				DefaultMapConfig += TEXT("EditorStartupMap=/Game/MobileStarterContent/Maps/Minimal_Default") LINE_TERMINATOR;
+				DefaultMapConfig += TEXT("GameDefaultMap=/Game/MobileStarterContent/Maps/Minimal_Default") LINE_TERMINATOR;
+			}
+			else
+			{
+				DefaultMapConfig += TEXT("EditorStartupMap=/Game/StarterContent/Maps/Minimal_Default") LINE_TERMINATOR;
+				DefaultMapConfig += TEXT("GameDefaultMap=/Game/StarterContent/Maps/Minimal_Default") LINE_TERMINATOR;
+			}
+		}
+
+		return MoveTemp(DefaultMapConfig);
+	}
 } // namespace <>
 
 FText FNewClassInfo::GetClassName() const
@@ -795,13 +819,11 @@ bool GameProjectUtils::OpenCodeIDE(const FString& ProjectFile, FText& OutFailRea
 	return true;
 }
 
-void GameProjectUtils::GetStarterContentFiles(TArray<FString>& OutFilenames)
+bool GameProjectUtils::IsStarterContentAvailableForNewProjects()
 {
-	FString const SrcFolder = FPaths::FeaturePackDir();
-
-	FString SearchPath = TEXT("*");
-	SearchPath += DefaultFeaturePackExtension;
-	IFileManager::Get().FindFilesRecursive(OutFilenames, *SrcFolder, *SearchPath, /*Files=*/true, /*Directories=*/false);
+	TArray<FString> OutFilenames;
+	IFileManager::Get().FindFilesRecursive(OutFilenames, *FPaths::FeaturePackDir(), TEXT("*StarterContent.upack"), /*Files=*/true, /*Directories=*/false);
+	return OutFilenames.Num() > 0;
 }
 
 bool GameProjectUtils::CreateProject(const FProjectInformation& InProjectInfo, FText& OutFailReason, FText& OutFailLog, TArray<FString>* OutCreatedFiles)
@@ -1316,22 +1338,6 @@ UTemplateProjectDefs* GameProjectUtils::LoadTemplateDefs(const FString& ProjectD
 		{
 			UE_LOG(LogGameProjectGeneration, Warning, TEXT("Template '%s' contains 'All' in HiddenSettings in addition to other entries. This is a mistake, and means that all settings will be hidden."), *ProjectDirectory);
 		}
-
-		TArray<TSharedPtr<FTemplateCategory>> AllTemplateCategories;
-		FGameProjectGenerationModule::Get().GetAllTemplateCategories(AllTemplateCategories);
-
-		for (const FName& CategoryKey : TemplateDefs->Categories)
-		{
-			bool bCategoryExists = AllTemplateCategories.ContainsByPredicate([&CategoryKey](const TSharedPtr<FTemplateCategory>& Category)
-				{
-					return Category->Key == CategoryKey;
-				});
-
-			if (!bCategoryExists)
-			{
-				UE_LOG(LogGameProjectGeneration, Warning, TEXT("Failed to find category definition named '%s' while loading template '%s', it is not defined in any TemplateCategories.ini."), *CategoryKey.ToString(), *ProjectDirectory);
-			}
-		}
 	}
 
 	return TemplateDefs;
@@ -1632,7 +1638,7 @@ bool GameProjectUtils::CreateProjectFromTemplate(const FProjectInformation& InPr
 
 	const FString ProjectConfigPath = DestFolder / TEXT("Config");
 
-	// Write out the hardware class target settings chosen for this project
+	// Write out DefaultEngine.ini settings
 	{
 		const FString DefaultEngineIniFilename = ProjectConfigPath / TEXT("DefaultEngine.ini");
 
@@ -1640,12 +1646,13 @@ bool GameProjectUtils::CreateProjectFromTemplate(const FProjectInformation& InPr
 		// Load the existing file - if it doesn't exist we create it
 		FFileHelper::LoadFileToString(FileContents, *DefaultEngineIniFilename);
 
-		FileContents += LINE_TERMINATOR;
 		FileContents += GetHardwareConfigString(InProjectInfo);
 
 		FileContents += GetRaytracingConfigString(InProjectInfo);
 
-		if ( !WriteOutputFile(DefaultEngineIniFilename, FileContents, OutFailReason) )
+		FileContents += GetDefaultMapConfigString(InProjectInfo);
+
+		if (!WriteOutputFile(DefaultEngineIniFilename, FileContents, OutFailReason))
 		{
 			return false;
 		}
@@ -1973,10 +1980,10 @@ FString GameProjectUtils::GetHardwareConfigString(const FProjectInformation& InP
 	FString GraphicsPresetAsString;
 	UEnum::GetValueAsString(TEXT("/Script/HardwareTargeting.EGraphicsPreset"), InProjectInfo.DefaultGraphicsPerformance, /*out*/ GraphicsPresetAsString);
 
+	HardwareTargeting += LINE_TERMINATOR;
 	HardwareTargeting += TEXT("[/Script/HardwareTargeting.HardwareTargetingSettings]") LINE_TERMINATOR;
 	HardwareTargeting += FString::Printf(TEXT("TargetedHardwareClass=%s") LINE_TERMINATOR, *TargetHardwareAsString);
 	HardwareTargeting += FString::Printf(TEXT("DefaultGraphicsPerformance=%s") LINE_TERMINATOR, *GraphicsPresetAsString);
-	HardwareTargeting += LINE_TERMINATOR;
 
 	return MoveTemp(HardwareTargeting);
 }
@@ -1993,58 +2000,49 @@ bool GameProjectUtils::GenerateConfigFiles(const FProjectInformation& InProjectI
 		const FString DefaultEngineIniFilename = ProjectConfigPath / TEXT("DefaultEngine.ini");
 		FString FileContents;
 
-		FileContents += TEXT("[URL]") LINE_TERMINATOR;
-
 		FileContents += GetHardwareConfigString(InProjectInfo);
-		FileContents += LINE_TERMINATOR;
+
+		FileContents += GetRaytracingConfigString(InProjectInfo);
 
 		if(bUseAudioMixerForAllPlatforms)
 		{
+			FileContents += LINE_TERMINATOR;
 			FileContents += TEXT("[Audio]") LINE_TERMINATOR;
 			FileContents += TEXT("UseAudioMixer=True") LINE_TERMINATOR;
-			FileContents += LINE_TERMINATOR;
 		}
 
 		if (InProjectInfo.bForceExtendedLuminanceRange)
 		{
+			FileContents += LINE_TERMINATOR;
 			FileContents += TEXT("[/Script/Engine.RendererSettings]") LINE_TERMINATOR;
 			FileContents += TEXT("r.DefaultFeature.AutoExposure.ExtendDefaultLuminanceRange=True") LINE_TERMINATOR;
-			FileContents += LINE_TERMINATOR;
 		}
 
 		if (InProjectInfo.bCopyStarterContent)
 		{
-			FString SpecificEditorStartupMap;
-			FString SpecificGameDefaultMap;
+			FileContents += LINE_TERMINATOR;
+			FileContents += TEXT("[/Script/EngineSettings.GameMapsSettings]") LINE_TERMINATOR;
 
-			// If we have starter content packs available, specify starter map
-			if( IsStarterContentAvailableForNewProjects() == true )
+			if (GameProjectUtils::IsStarterContentAvailableForNewProjects())
 			{
 				if (InProjectInfo.TargetedHardware == EHardwareClass::Mobile)
 				{
-					SpecificEditorStartupMap = TEXT("/Game/MobileStarterContent/Maps/Minimal_Default");
-					SpecificGameDefaultMap = TEXT("/Game/MobileStarterContent/Maps/Minimal_Default");
+					FileContents += TEXT("EditorStartupMap=/Game/MobileStarterContent/Maps/Minimal_Default") LINE_TERMINATOR;
+					FileContents += TEXT("GameDefaultMap=/Game/MobileStarterContent/Maps/Minimal_Default") LINE_TERMINATOR;
 				}
 				else
 				{
-					SpecificEditorStartupMap = TEXT("/Game/StarterContent/Maps/Minimal_Default");
-					SpecificGameDefaultMap = TEXT("/Game/StarterContent/Maps/Minimal_Default");
+					FileContents += TEXT("EditorStartupMap=/Game/StarterContent/Maps/Minimal_Default") LINE_TERMINATOR;
+					FileContents += TEXT("GameDefaultMap=/Game/StarterContent/Maps/Minimal_Default") LINE_TERMINATOR;
 				}
 			}
 
-			// Write out the settings for startup map and game default map
-			FileContents += TEXT("[/Script/EngineSettings.GameMapsSettings]") LINE_TERMINATOR;
-			FileContents += FString::Printf(TEXT("EditorStartupMap=%s") LINE_TERMINATOR, *SpecificEditorStartupMap);
-			FileContents += FString::Printf(TEXT("GameDefaultMap=%s") LINE_TERMINATOR, *SpecificGameDefaultMap);
 			if (InProjectInfo.bShouldGenerateCode)
 			{
 				FileContents += FString::Printf(TEXT("GlobalDefaultGameMode=\"/Script/%s.%sGameMode\"") LINE_TERMINATOR, *NewProjectName, *NewProjectName);
 			}
 		}
 
-		FileContents += LINE_TERMINATOR;
-		FileContents += GetRaytracingConfigString(InProjectInfo);
-
 		if (WriteOutputFile(DefaultEngineIniFilename, FileContents, OutFailReason))
 		{
 			OutCreatedFiles.Add(DefaultEngineIniFilename);
@@ -2364,15 +2362,6 @@ bool GameProjectUtils::GenerateCodeProjectFiles(const FString& ProjectFilename,
 	return true;
 }
 
-bool GameProjectUtils::IsStarterContentAvailableForNewProjects()
-{
-	TArray<FString> StarterContentFiles;
-	GetStarterContentFiles(StarterContentFiles);
-
-	bool bHasStaterContent = StarterContentFiles.FindByPredicate([&](const FString& Str){ return Str.Contains("StarterContent"); }) != nullptr;
-	return bHasStaterContent;
-}
-
 int32 FindSpecificBuildFiles(const TCHAR* Path, TMap<FString, FString>& BuildFiles)
 {
 	class FBuildFileVistor : public IPlatformFile::FDirectoryVisitor
@@ -4089,11 +4078,11 @@ bool GameProjectUtils::InsertFeaturePacksIntoINIFile(const FProjectInformation&
 		FString StarterPack;
 		if (InProjectInfo.TargetedHardware == EHardwareClass::Mobile)
 		{
-			StarterPack = TEXT("InsertPack=(PackSource=\"MobileStarterContent") + DefaultFeaturePackExtension + TEXT("\",PackName=\"StarterContent\")");
+			StarterPack = TEXT("InsertPack=(PackSource=\"MobileStarterContent.upack\",PackName=\"StarterContent\")");
 		}
 		else
 		{
-			StarterPack = TEXT("InsertPack=(PackSource=\"StarterContent")  + DefaultFeaturePackExtension + TEXT("\",PackName=\"StarterContent\")");
+			StarterPack = TEXT("InsertPack=(PackSource=\"StarterContent.upack\",PackName=\"StarterContent\")");
 		}
 		PackList.Add(StarterPack);
 	}
@@ -4133,9 +4122,6 @@ bool GameProjectUtils::AddSharedContentToProject(const FProjectInformation &InPr
 	const FString SrcFolder = FPaths::GetPath(InProjectInfo.TemplateFile);
 	const FString DestFolder = FPaths::GetPath(InProjectInfo.ProjectFilename);
 
-	const FString ProjectConfigPath = DestFolder / TEXT("Config");
-	const FString IniFilename = ProjectConfigPath / TEXT("DefaultGame.ini");
-
 	// Now any packs specified in the template def.
 	UTemplateProjectDefs* TemplateDefs = LoadTemplateDefs(SrcFolder);
 	if (TemplateDefs != NULL)
@@ -4147,7 +4133,7 @@ bool GameProjectUtils::AddSharedContentToProject(const FProjectInformation &InPr
 		}
 
 		TUniquePtr<FFeaturePackContentSource> TempFeaturePack = MakeUnique<FFeaturePackContentSource>();
-		bool bCopied = TempFeaturePack->InsertAdditionalResources(TemplateDefs->SharedContentPacks,RequiredDetail, DestFolder,CreatedFiles);
+		bool bCopied = TempFeaturePack->InsertAdditionalResources(TemplateDefs->SharedContentPacks, RequiredDetail, DestFolder, CreatedFiles);
 		if( bCopied == false )
 		{
 			FFormatNamedArguments Args;
diff --git a/Engine/Source/Editor/GameProjectGeneration/Private/SGameProjectDialog.cpp b/Engine/Source/Editor/GameProjectGeneration/Private/SGameProjectDialog.cpp
--- a/Engine/Source/Editor/GameProjectGeneration/Private/SGameProjectDialog.cpp
+++ b/Engine/Source/Editor/GameProjectGeneration/Private/SGameProjectDialog.cpp
@@ -50,7 +50,7 @@ namespace GameProjectDialogDefs
 	constexpr int32 ProjectSettingsPageIndex = 3;
 	
 	static const FText LandingPageTitle = LOCTEXT("ProjectDialog_SelectProject", "Select Project");
-	static const FText SelectCategoryTitle = LOCTEXT("ProjectDialog_SelectCategory", "Select Category");
+	static const FText SelectCategoryTitle = LOCTEXT("ProjectDialog_SelectTemplateCategory", "Select Template Category");
 	static const FText ProjectBrowserTitle = LOCTEXT("ProjectDialog_ProjectBrowser", "Project Browser");
 	static const FText TemplateListTitle = LOCTEXT("ProjectDialog_SelectTemplate", "Select Template");
 	static const FText ProjectSettingsTitle = LOCTEXT("ProjectDialog_ProjectSettings", "Project Settings");
@@ -63,7 +63,8 @@ void SGameProjectDialog::Construct(const FArguments& InArgs, EMode InMode)
 
 	NewProjectWizard = SNew(SNewProjectWizard)
 		.OnTemplateDoubleClick(this, &SGameProjectDialog::OnTemplateDoubleClick);
-	ProjectBrowserPage = SNew(SProjectBrowser);
+	ProjectBrowserPage = SNew(SProjectBrowser)
+		.HideOpenButton(true);
 
 	const float UniformPadding = 16.0f;
 
@@ -122,7 +123,7 @@ void SGameProjectDialog::Construct(const FArguments& InArgs, EMode InMode)
 	];
 }
 
-TSharedRef<SWidget> SGameProjectDialog::CreateLandingPage()
+void SGameProjectDialog::GetAllTemplateCategories(TArray<TSharedPtr<FTemplateCategory>>& OutCategories)
 {
 	TArray<TSharedPtr<FTemplateCategory>> AllTemplateCategories;
 	FGameProjectGenerationModule::Get().GetAllTemplateCategories(AllTemplateCategories);
@@ -140,6 +141,14 @@ TSharedRef<SWidget> SGameProjectDialog::CreateLandingPage()
 		AllTemplateCategories.Add(DefaultCategory);
 	}
 
+	OutCategories = AllTemplateCategories;
+}
+
+TSharedRef<SWidget> SGameProjectDialog::CreateLandingPage()
+{
+	TArray<TSharedPtr<FTemplateCategory>> AllTemplateCategories;
+	SGameProjectDialog::GetAllTemplateCategories(AllTemplateCategories);
+
 	for (const TSharedPtr<FTemplateCategory>& Category : AllTemplateCategories)
 	{
 		if (Category->IsMajor)
@@ -152,32 +161,8 @@ TSharedRef<SWidget> SGameProjectDialog::CreateLandingPage()
 		}
 	}
 
-	FText Description;
-
-	if (DialogMode == EMode::New)
-	{
-		Description = LOCTEXT("LandingPageDesc_New", "Choose a template <RichTextBlock.BoldHighlight>category</> for your new project.");
-	}
-	else if (DialogMode == EMode::Open)
-	{
-		Description = LOCTEXT("LandingPageDesc_Open", "Choose a recent <RichTextBlock.BoldHighlight>project</> to load, or press <RichTextBlock.BoldHighlight>More</> to see all projects on your computer.");
-	}
-	else if (DialogMode == EMode::Both)
-	{
-		Description = LOCTEXT("LandingPageDesc_Both", "Choose a recent <RichTextBlock.BoldHighlight>project</> to load, press <RichTextBlock.BoldHighlight>More</> to see all projects on your computer, or choose a template <RichTextBlock.BoldHighlight>category</> to create a new project.");
-	}
-
 	TSharedRef<SWidget> Widget = SNew(SVerticalBox)
 	+ SVerticalBox::Slot()
-	.AutoHeight()
-	.Padding(0, 9, 0, 8)
-	[
-		SNew(SRichTextBlock)
-		.Text(Description)
-		.AutoWrapText(true)
-		.DecoratorStyleSet(&FEditorStyle::Get())
-	]
-	+ SVerticalBox::Slot()
 	.FillHeight(1)
 	[
 		SNew(SBorder)
@@ -226,8 +211,8 @@ TSharedRef<SWidget> SGameProjectDialog::CreateLandingPage()
 					SNew(STextBlock)
 					.TextStyle(FEditorStyle::Get(), "GameProjectDialog.ProjectNamePathLabels")
 					.Text(LOCTEXT("ProjectDialog_Categories", "Categories"))
+					.Visibility(DialogMode == EMode::Both ? EVisibility::Visible : EVisibility::Collapsed)
 				]
-
 				+ SVerticalBox::Slot()
 				.AutoHeight()
 				.Padding(8, 8, 8, 0)
diff --git a/Engine/Source/Editor/GameProjectGeneration/Private/SGameProjectDialog.h b/Engine/Source/Editor/GameProjectGeneration/Private/SGameProjectDialog.h
--- a/Engine/Source/Editor/GameProjectGeneration/Private/SGameProjectDialog.h
+++ b/Engine/Source/Editor/GameProjectGeneration/Private/SGameProjectDialog.h
@@ -43,6 +43,8 @@ public:
 	/** Constructs this widget with InArgs */
 	void Construct( const FArguments& InArgs, EMode Mode );
 
+	static void GetAllTemplateCategories(TArray<TSharedPtr<FTemplateCategory>>& OutCategories);
+
 private:
 	EMode DialogMode;
 
diff --git a/Engine/Source/Editor/GameProjectGeneration/Private/SNewProjectWizard.cpp b/Engine/Source/Editor/GameProjectGeneration/Private/SNewProjectWizard.cpp
--- a/Engine/Source/Editor/GameProjectGeneration/Private/SNewProjectWizard.cpp
+++ b/Engine/Source/Editor/GameProjectGeneration/Private/SNewProjectWizard.cpp
@@ -1,50 +1,50 @@
 // Copyright 1998-2019 Epic Games, Inc. All Rights Reserved.
 
-
 #include "SNewProjectWizard.h"
 #include "Brushes/SlateDynamicImageBrush.h"
-#include "HAL/PlatformFilemanager.h"
-#include "Misc/MessageDialog.h"
+#include "DesktopPlatformModule.h"
+#include "Dialogs/SOutputLogDialog.h"
+#include "Editor.h"
+#include "EditorStyleSet.h"
+#include "Framework/Application/SlateApplication.h"
+#include "GameProjectGenerationLog.h"
+#include "GameProjectGenerationModule.h"
+#include "GameProjectUtils.h"
 #include "HAL/FileManager.h"
+#include "HAL/PlatformFilemanager.h"
+#include "HardwareTargetingModule.h"
+#include "IDocumentation.h"
+#include "Interfaces/IPluginManager.h"
+#include "Interfaces/IProjectManager.h"
+#include "Internationalization/BreakIterator.h"
+#include "Layout/WidgetPath.h"
 #include "Misc/App.h"
 #include "Misc/CommandLine.h"
+#include "Misc/MessageDialog.h"
 #include "Misc/Parse.h"
-#include "Widgets/SOverlay.h"
-#include "Layout/WidgetPath.h"
+#include "ProjectDescriptor.h"
+#include "SGameProjectDialog.h"
+#include "SGetSuggestedIDEWidget.h"
+#include "Settings/EditorSettings.h"
 #include "SlateOptMacros.h"
-#include "Framework/Application/SlateApplication.h"
+#include "SourceCodeNavigation.h"
+#include "TemplateCategory.h"
+#include "TemplateItem.h"
+#include "TemplateProjectDefs.h"
 #include "Textures/SlateIcon.h"
-#include "Widgets/Layout/SSeparator.h"
 #include "Widgets/Images/SImage.h"
-#include "Widgets/Text/SRichTextBlock.h"
-#include "Widgets/Input/SEditableTextBox.h"
 #include "Widgets/Input/SButton.h"
+#include "Widgets/Input/SCheckBox.h"
+#include "Widgets/Input/SEditableTextBox.h"
 #include "Widgets/Layout/SGridPanel.h"
 #include "Widgets/Layout/SScrollBorder.h"
 #include "Widgets/Layout/SScrollBox.h"
+#include "Widgets/Layout/SSeparator.h"
 #include "Widgets/Layout/SWrapBox.h"
-#include "Widgets/Input/SCheckBox.h"
-#include "EditorStyleSet.h"
-#include "Editor.h"
-#include "Interfaces/IPluginManager.h"
-#include "Interfaces/IProjectManager.h"
-#include "ProjectDescriptor.h"
-#include "GameProjectGenerationLog.h"
-#include "GameProjectGenerationModule.h"
-#include "HardwareTargetingModule.h"
-#include "TemplateProjectDefs.h"
-#include "GameProjectUtils.h"
-#include "SGetSuggestedIDEWidget.h"
-#include "DesktopPlatformModule.h"
-#include "SourceCodeNavigation.h"
-#include "TemplateCategory.h"
+#include "Widgets/SOverlay.h"
 #include "Widgets/SToolTip.h"
+#include "Widgets/Text/SRichTextBlock.h"
 #include "Widgets/Workflow/SWizard.h"
-#include "IDocumentation.h"
-#include "Internationalization/BreakIterator.h"
-#include "Dialogs/SOutputLogDialog.h"
-#include "TemplateItem.h"
-#include "Settings/EditorSettings.h"
 
 #define LOCTEXT_NAMESPACE "NewProjectWizard"
 
@@ -305,11 +305,8 @@ void SNewProjectWizard::Construct( const FArguments& InArgs )
 {
 	bProjectSettingsHidden = false;
 
-	LastValidityCheckTime = 0;
-	ValidityCheckFrequency = 4;
 	bLastGlobalValidityCheckSuccessful = true;
 	bLastNameAndLocationValidityCheckSuccessful = true;
-	bPreventPeriodicValidityChecksUntilNextChange = false;
 	bCopyStarterContent = GEditor ? GetDefault<UEditorSettings>()->bCopyStarterContentPreference : true;
 
 	SelectedHardwareClassTarget = EHardwareClass::Desktop;
@@ -321,13 +318,13 @@ void SNewProjectWizard::Construct( const FArguments& InArgs )
 	FindTemplateProjects();
 	SetDefaultProjectLocation();
 
-	TemplateListView = SNew(STileView< TSharedPtr<FTemplateItem> >)
+	TemplateListView = SNew(STileView<TSharedPtr<FTemplateItem>>)
 	.ListItemsSource(&FilteredTemplateList)
 	.SelectionMode(ESelectionMode::Single)
 	.ClearSelectionOnClick(false)
 	.OnGenerateTile_Static(&STemplateTile::BuildTile)
-	.ItemHeight( NewProjectWizardDefs::ItemHeight )
-	.ItemWidth( NewProjectWizardDefs::ItemWidth )
+	.ItemHeight(NewProjectWizardDefs::ItemHeight)
+	.ItemWidth(NewProjectWizardDefs::ItemWidth)
 	.OnMouseButtonDoubleClick(this, &SNewProjectWizard::HandleTemplateListViewDoubleClick)
 	.OnSelectionChanged(this, &SNewProjectWizard::HandleTemplateListViewSelectionChanged);
 
@@ -341,17 +338,6 @@ void SNewProjectWizard::Construct( const FArguments& InArgs )
 	[
 		SNew(SVerticalBox)
 
-		+ SVerticalBox::Slot()
-		.Padding(0, 8, 0, 8)
-		.AutoHeight()
-		[
-			SNew(SRichTextBlock)
-			.Text(LOCTEXT("ProjectTemplateDescription", "Choose a <RichTextBlock.BoldHighlight>template</> to use as a starting point for your new project.  Any of these features can be added later by clicking <RichTextBlock.BoldHighlight>Add Feature or Content Pack</> in <RichTextBlock.BoldHighlight>Content Browser</>."))
-			.AutoWrapText(true)
-			.DecoratorStyleSet(&FEditorStyle::Get())
-			.ToolTip(IDocumentation::Get()->CreateToolTip(LOCTEXT("TemplateChoiceTooltip", "A template consists of a little bit of player control logic (either as a Blueprint or in C++), input bindings, and appropriate prototyping assets."), nullptr, TEXT("Shared/Editor/NewProjectWizard"), TEXT("TemplateChoice")))
-		]
-
 		+ SVerticalBox::Slot()
 		[
 			SNew(SVerticalBox)
@@ -479,31 +465,6 @@ void SNewProjectWizard::Construct( const FArguments& InArgs )
 									]
 								]
 							]
-
-							// Languages available
-							+ SVerticalBox::Slot()
-							.AutoHeight()
-							.Padding(FMargin(0.0f, 5.0f, 0.0f, 5.0f))
-							[
-								SNew(SBox)
-								.Visibility(this, &SNewProjectWizard::GetSelectedTemplateLanguageVisibility)
-								[
-									SNew(SVerticalBox)
-									+ SVerticalBox::Slot()
-									[
-										SNew(STextBlock)
-										.TextStyle(FEditorStyle::Get(), "GameProjectDialog.FeatureText")
-										.Text(LOCTEXT("LanguagesAvailable", "Languages Available:"))
-									]
-									+ SVerticalBox::Slot()
-									.AutoHeight()
-									[
-										SNew(STextBlock)
-										.AutoWrapText(true)
-										.Text(this, &SNewProjectWizard::GetSelectedTemplateLanguages)
-									]
-								]
-							]
 						]
 					]
 				]
@@ -560,17 +521,6 @@ FText SNewProjectWizard::GetStarterContentWarningTooltip() const
 	}
 }
 
-void SNewProjectWizard::Tick( const FGeometry& AllottedGeometry, const double InCurrentTime, const float InDeltaTime )
-{
-	// Every few seconds, the project file path is checked for validity in case the disk contents changed and the location is now valid or invalid.
-	// After project creation, periodic checks are disabled to prevent a brief message indicating that the project you created already exists.
-	// This feature is re-enabled if the user did not restart and began editing parameters again.
-	if ( !bPreventPeriodicValidityChecksUntilNextChange && (InCurrentTime > LastValidityCheckTime + ValidityCheckFrequency) )
-	{
-		UpdateProjectFileValidity();
-	}
-}
-
 void SNewProjectWizard::HandleTemplateListViewSelectionChanged(TSharedPtr<FTemplateItem> TemplateItem, ESelectInfo::Type SelectInfo)
 {
 	UpdateProjectFileValidity();
@@ -627,48 +577,6 @@ EVisibility SNewProjectWizard::GetSelectedTemplatePreviewVisibility() const
 	return PreviewImage.IsValid() ? EVisibility::Visible : EVisibility::Collapsed;
 }
 
-FText SNewProjectWizard::GetSelectedTemplateLanguages() const
-{
-	TSharedPtr<FTemplateItem> SelectedItem = GetSelectedTemplateItem();
-	if (SelectedItem.IsValid())
-	{
-		bool bHasBlueprint = !SelectedItem->BlueprintProjectFile.IsEmpty();
-		bool bHasCode = !SelectedItem->CodeProjectFile.IsEmpty();
-
-		if (!bHasBlueprint && !bHasCode)
-		{
-			bHasBlueprint = true;
-			bHasCode = true;
-		}
-
-		if (bHasBlueprint && bHasCode)
-		{
-			return LOCTEXT("ProjectDialog_BlueprintAndCode", "Blueprint, C++");
-		}
-		else if (bHasBlueprint)
-		{
-			return LOCTEXT("ProjectDialog_Blueprint", "Blueprint");
-		}
-		else if (bHasCode)
-		{
-			return LOCTEXT("ProjectDialog_Code", "C++");
-		}
-	}
-
-	return FText::GetEmpty();
-}
-
-EVisibility SNewProjectWizard::GetSelectedTemplateLanguageVisibility() const
-{
-	FText Languages = GetSelectedTemplateLanguages();
-	if (Languages.IsEmpty())
-	{
-		return EVisibility::Collapsed;
-	}
-
-	return EVisibility::Visible;
-}
-
 FText SNewProjectWizard::GetCurrentProjectFileName() const
 {
 	return FText::FromString( CurrentProjectFileName );
@@ -801,7 +709,6 @@ FText SNewProjectWizard::GetGlobalErrorLabelText() const
 	return FText::GetEmpty();
 }
 
-
 FReply SNewProjectWizard::OnCloseGlobalErrorLabelClicked()
 {
 	PersistentGlobalErrorLabelText = FText();
@@ -809,13 +716,11 @@ FReply SNewProjectWizard::OnCloseGlobalErrorLabelClicked()
 	return FReply::Handled();
 }
 
-
 EVisibility SNewProjectWizard::GetNameAndLocationErrorLabelVisibility() const
 {
 	return GetNameAndLocationErrorLabelText().IsEmpty() ? EVisibility::Collapsed : EVisibility::Visible;
 }
 
-
 FText SNewProjectWizard::GetNameAndLocationErrorLabelText() const
 {
 	if ( !bLastNameAndLocationValidityCheckSuccessful )
@@ -932,6 +837,7 @@ TMap<FName, TArray<TSharedPtr<FTemplateItem>> >& SNewProjectWizard::FindTemplate
 					Template->AssetTypes = TemplateDefs->AssetTypes;
 					Template->HiddenSettings = TemplateDefs->HiddenSettings;
 					Template->bIsEnterprise = TemplateDefs->bIsEnterprise;
+					Template->bIsBlankTemplate = TemplateDefs->bIsBlank;
 
 					Template->Name = TemplateDefs->GetDisplayNameText();
 					if (Template->Name.IsEmpty())
@@ -997,6 +903,26 @@ TMap<FName, TArray<TSharedPtr<FTemplateItem>> >& SNewProjectWizard::FindTemplate
 		}
 	}
 
+	TArray<TSharedPtr<FTemplateCategory>> AllTemplateCategories;
+	SGameProjectDialog::GetAllTemplateCategories(AllTemplateCategories);
+
+	// Validate that all our templates have a category defined
+	TArray<FName> CategoryKeys;
+	Templates.GetKeys(CategoryKeys);
+	for (const FName& CategoryKey : CategoryKeys)
+	{
+		bool bCategoryExists = AllTemplateCategories.ContainsByPredicate([&CategoryKey](const TSharedPtr<FTemplateCategory>& Category)
+			{
+				return Category->Key == CategoryKey;
+			});
+
+		if (!bCategoryExists)
+		{
+			UE_LOG(LogGameProjectGeneration, Warning, TEXT("Failed to find category definition named '%s', it is not defined in any TemplateCategories.ini."), *CategoryKey.ToString());
+		}
+	}
+
+
 	if (Templates.Num() == 0)
 	{
 		static const FName DefaultCategory("Default");
@@ -1015,9 +941,7 @@ TMap<FName, TArray<TSharedPtr<FTemplateItem>> >& SNewProjectWizard::FindTemplate
 		BlankTemplate->BlueprintProjectFile = TEXT("");
 		BlankTemplate->CodeProjectFile = TEXT("");
 		BlankTemplate->bIsEnterprise = false;
-
-		TArray<TSharedPtr<FTemplateCategory>> AllTemplateCategories;
-		FGameProjectGenerationModule::Get().GetAllTemplateCategories(AllTemplateCategories);
+		BlankTemplate->bIsBlankTemplate = true;
 
 		for (const TSharedPtr<FTemplateCategory>& Category : AllTemplateCategories)
 		{
@@ -1096,7 +1020,7 @@ void SNewProjectWizard::SetDefaultProjectLocation()
 	}
 }
 
-void SNewProjectWizard::UpdateProjectFileValidity( )
+void SNewProjectWizard::UpdateProjectFileValidity()
 {
 	// Global validity
 	{
@@ -1160,11 +1084,6 @@ void SNewProjectWizard::UpdateProjectFileValidity( )
 			}
 		}
 	}
-
-	LastValidityCheckTime = FSlateApplication::Get().GetCurrentTime();
-
-	// Since this function was invoked, periodic validity checks should be re-enabled if they were disabled.
-	bPreventPeriodicValidityChecksUntilNextChange = false;
 }
 
 bool SNewProjectWizard::IsCompilerRequired() const
@@ -1178,7 +1097,7 @@ bool SNewProjectWizard::IsCompilerRequired() const
 	return false;
 }
 
-bool SNewProjectWizard::CreateProject( const FString& ProjectFile )
+bool SNewProjectWizard::CreateProject(const FString& ProjectFile)
 {
 	// Get the selected template
 	TSharedPtr<FTemplateItem> SelectedTemplate = GetSelectedTemplateItem();
@@ -1203,6 +1122,7 @@ bool SNewProjectWizard::CreateProject( const FString& ProjectFile )
 	ProjectInfo.bEnableXR = bEnableXR;
 	ProjectInfo.bEnableRaytracing = bEnableRaytracing;
 	ProjectInfo.bIsEnterpriseProject = SelectedTemplate->bIsEnterprise;
+	ProjectInfo.bIsBlankTemplate = SelectedTemplate->bIsBlankTemplate;
 
 	if (!GameProjectUtils::CreateProject(ProjectInfo, FailReason, FailLog))
 	{
@@ -1228,22 +1148,19 @@ bool SNewProjectWizard::CreateProject( const FString& ProjectFile )
 	return true;
 }
 
-void SNewProjectWizard::CreateAndOpenProject( )
+void SNewProjectWizard::CreateAndOpenProject()
 {
-	if( !CanCreateProject() )
+	if (!CanCreateProject())
 	{
 		return;
 	}
 
 	FString ProjectFile = GetProjectFilenameWithPath();
-	if ( !CreateProject(ProjectFile) )
+	if (!CreateProject(ProjectFile))
 	{
 		return;
 	}
 
-	// Prevent periodic validity checks. This is to prevent a brief error message about the project already existing while you are exiting.
-	bPreventPeriodicValidityChecksUntilNextChange = true;
-
 	if (bShouldGenerateCode)
 	{
 	    // If the engine is installed it is already compiled, so we can try to build and open a new project immediately. Non-installed situations might require building
@@ -1252,8 +1169,8 @@ void SNewProjectWizard::CreateAndOpenProject( )
 		{
 			if (GameProjectUtils::BuildCodeProject(ProjectFile))
 			{
-				OpenCodeIDE( ProjectFile );
-				OpenProject( ProjectFile );
+				OpenCodeIDE(ProjectFile);
+				OpenProject(ProjectFile);
 			}
 			else
 			{
@@ -1262,19 +1179,19 @@ void SNewProjectWizard::CreateAndOpenProject( )
 		}
 		else
 		{
-			OpenCodeIDE( ProjectFile );
+			OpenCodeIDE(ProjectFile);
 		}
 	}
 	else
 	{
-		OpenProject( ProjectFile );
+		OpenProject(ProjectFile );
 	}
 }
 
-bool SNewProjectWizard::OpenProject( const FString& ProjectFile )
+bool SNewProjectWizard::OpenProject(const FString& ProjectFile)
 {
 	FText FailReason;
-	if ( GameProjectUtils::OpenProject( ProjectFile, FailReason ) )
+	if (GameProjectUtils::OpenProject(ProjectFile, FailReason))
 	{
 		// Successfully opened the project, the editor is closing.
 		// Close this window in case something prevents the editor from closing (save dialog, quit confirmation, etc)
@@ -1286,11 +1203,11 @@ bool SNewProjectWizard::OpenProject( const FString& ProjectFile )
 	return false;
 }
 
-bool SNewProjectWizard::OpenCodeIDE( const FString& ProjectFile )
+bool SNewProjectWizard::OpenCodeIDE(const FString& ProjectFile)
 {
 	FText FailReason;
 
-	if ( GameProjectUtils::OpenCodeIDE( ProjectFile, FailReason ) )
+	if (GameProjectUtils::OpenCodeIDE(ProjectFile, FailReason))
 	{
 		// Successfully opened code editing IDE, the editor is closing
 		// Close this window in case something prevents the editor from closing (save dialog, quit confirmation, etc)
@@ -1298,7 +1215,7 @@ bool SNewProjectWizard::OpenCodeIDE( const FString& ProjectFile )
 		return true;
 	}
 
-	DisplayError( FailReason );
+	DisplayError(FailReason);
 	return false;
 }
 
@@ -1352,7 +1269,7 @@ void SNewProjectWizard::SetCurrentCategory(FName Category)
 
 TSharedRef<SWidget> SNewProjectWizard::MakeProjectLocationWidget()
 {
-	return SNew(SVerticalBox)
+	TSharedRef<SWidget> Widget = SNew(SVerticalBox)
 	+ SVerticalBox::Slot()
 	.HAlign(HAlign_Center)
 	.Padding(8)
@@ -1383,6 +1300,15 @@ TSharedRef<SWidget> SNewProjectWizard::MakeProjectLocationWidget()
 			.OnNameChanged(this, &SNewProjectWizard::OnCurrentProjectFileNameChanged)
 		]
 	];
+
+	Widget->RegisterActiveTimer(1.0f, FWidgetActiveTimerDelegate::CreateLambda(
+		[this](double, float)
+		{
+			UpdateProjectFileValidity();
+			return EActiveTimerReturnType::Continue;
+		}));
+
+	return Widget;
 }
 
 TSharedRef<SWidget> SNewProjectWizard::CreateProjectSettingsPage()
@@ -1394,12 +1320,6 @@ TSharedRef<SWidget> SNewProjectWizard::CreateProjectSettingsPage()
 	[
 		SNew(SVerticalBox)
 		+ SVerticalBox::Slot()
-		.AutoHeight()
-		.Padding(0, 8, 0, 8)
-		[
-			MakeProjectSettingsDescriptionBox()
-		]
-		+ SVerticalBox::Slot()
 		.Padding(0)
 		.VAlign(VAlign_Fill)
 		[
@@ -1542,17 +1462,6 @@ void SNewProjectWizard::OnSetBlueprintOrCppIndex(int32 Index)
 	bShouldGenerateCode = Index == 1;
 }
 
-TSharedRef<SRichTextBlock> SNewProjectWizard::MakeProjectSettingsDescriptionBox() const
-{
-	TSharedRef<SRichTextBlock> Widget = SNew(SRichTextBlock)
-		.Text(LOCTEXT("ProjectSettingsDescription", "Choose some <RichTextBlock.BoldHighlight>settings</> for your project."))
-		.AutoWrapText(true)
-		.DecoratorStyleSet(&FEditorStyle::Get())
-		.ToolTip(IDocumentation::Get()->CreateToolTip(LOCTEXT("HardwareTargetTooltip", "These settings will choose good defaults for a number of other settings in the project such as post-processing flags and touch input emulation using the mouse."), nullptr, TEXT("Shared/Editor/NewProjectWizard"), TEXT("TargetHardware")));
-
-	return Widget;
-}
-
 static void AddToProjectSettingsGrid(TSharedRef<SGridPanel> Grid, const TSharedRef<SWidget>& Enum, const TSharedRef<SWidget>& Description, FIntPoint& Slot)
 {
 	Grid->AddSlot(Slot.X, Slot.Y)
@@ -1699,7 +1608,7 @@ TSharedRef<SWidget> SNewProjectWizard::MakeProjectSettingsOptionsBox()
 
 		TSharedRef<SRichTextBlock> Description =
 			SNew(SRichTextBlock)
-			.Text(LOCTEXT("CopyStarterContent_ToolTip", "Enable to include an additional content pack containing simple placeable meshes with basic materials and textures.\nYou can opt out of including this to create a project that only has the bare essentials for the selected project template.\nYou can also add the <RichTextBlock.BoldHighlight>Starter Content</> to your project later using <RichTextBlock.BoldHighlight>Content Browser</>."))
+			.Text(LOCTEXT("CopyStarterContent_ToolTip", "Enable to include an additional content pack containing simple placeable meshes with basic materials and textures.\nYou can also add the <RichTextBlock.BoldHighlight>Starter Content</> to your project later using <RichTextBlock.BoldHighlight>Content Browser</>."))
 			.AutoWrapText(true)
 			.DecoratorStyleSet(&FEditorStyle::Get());
 		
diff --git a/Engine/Source/Editor/GameProjectGeneration/Private/SNewProjectWizard.h b/Engine/Source/Editor/GameProjectGeneration/Private/SNewProjectWizard.h
--- a/Engine/Source/Editor/GameProjectGeneration/Private/SNewProjectWizard.h
+++ b/Engine/Source/Editor/GameProjectGeneration/Private/SNewProjectWizard.h
@@ -33,8 +33,6 @@ public:
 	/** Constructs this widget with InArgs */
 	void Construct( const FArguments& InArgs );
 
-	virtual void Tick( const FGeometry& AllottedGeometry, const double InCurrentTime, const float InDeltaTime ) override;
-
 	/** Populates TemplateItemsSource with templates found on disk */
 	TMap<FName, TArray<TSharedPtr<FTemplateItem>> >& FindTemplateProjects();
 
@@ -102,12 +100,6 @@ private:
 	/** Get a string that details the class types referenced in the selected template */
 	FText GetSelectedTemplateClassTypes() const;
 
-	/** Get a string that details the languages supported by the selected template */
-	FText GetSelectedTemplateLanguages() const;
-
-	/** Get visibility of the languages display. */
-	EVisibility GetSelectedTemplateLanguageVisibility() const;
-	
 	/** Get a visiblity of the class types display. If the string is empty this return Collapsed otherwise it will return Visible */
 	EVisibility GetSelectedTemplateClassVisibility() const;
 	
@@ -211,7 +203,6 @@ private:
 	/** Handler for when the selection changes in the template list */
 	void HandleTemplateListViewSelectionChanged( TSharedPtr<FTemplateItem> TemplateItem, ESelectInfo::Type SelectInfo );
 
-	TSharedRef<SRichTextBlock> MakeProjectSettingsDescriptionBox() const;
 	TSharedRef<SWidget> MakeProjectSettingsOptionsBox();
 
 private:
@@ -223,15 +214,6 @@ private:
 
 	FSimpleDelegate OnTemplateDoubleClick;
 
-	/** The last time that the selected project file path was checked for validity. This is used to throttle I/O requests to a reasonable frequency */
-	double LastValidityCheckTime;
-
-	/** The frequency in seconds for validity checks while the dialog is idle. Changes to the path immediately update the validity. */
-	double ValidityCheckFrequency;
-
-	/** Periodic checks for validity will not occur while this flag is true. Used to prevent a frame of "this project already exists" while exiting after a successful creation. */
-	bool bPreventPeriodicValidityChecksUntilNextChange;
-
 	/** The global error text from the last validity check */
 	FText LastGlobalValidityErrorText;
 
diff --git a/Engine/Source/Editor/GameProjectGeneration/Private/SProjectBrowser.cpp b/Engine/Source/Editor/GameProjectGeneration/Private/SProjectBrowser.cpp
--- a/Engine/Source/Editor/GameProjectGeneration/Private/SProjectBrowser.cpp
+++ b/Engine/Source/Editor/GameProjectGeneration/Private/SProjectBrowser.cpp
@@ -324,7 +324,7 @@ void SProjectBrowser::Construct( const FArguments& InArgs )
 					SNew(SButton)
 					.Text(LOCTEXT("BrowseProjectButton", "Browse..."))
 					.OnClicked(this, &SProjectBrowser::OnBrowseToProjectClicked)
-					.ContentPadding( FCoreStyle::Get().GetMargin("StandardDialog.ContentPadding") )
+					.ContentPadding(FCoreStyle::Get().GetMargin("StandardDialog.ContentPadding"))
 				]
 
 				// Open Button
@@ -335,7 +335,8 @@ void SProjectBrowser::Construct( const FArguments& InArgs )
 					.Text(LOCTEXT("OpenProjectButton", "Open"))
 					.OnClicked(this, &SProjectBrowser::HandleOpenProjectButtonClicked)
 					.IsEnabled(this, &SProjectBrowser::HandleOpenProjectButtonIsEnabled)
-					.ContentPadding( FCoreStyle::Get().GetMargin("StandardDialog.ContentPadding") )
+					.ContentPadding(FCoreStyle::Get().GetMargin("StandardDialog.ContentPadding"))
+					.Visibility(InArgs._HideOpenButton ? EVisibility::Collapsed : EVisibility::Visible)
 				]
 			]
 		]
diff --git a/Engine/Source/Editor/GameProjectGeneration/Private/SProjectBrowser.h b/Engine/Source/Editor/GameProjectGeneration/Private/SProjectBrowser.h
--- a/Engine/Source/Editor/GameProjectGeneration/Private/SProjectBrowser.h
+++ b/Engine/Source/Editor/GameProjectGeneration/Private/SProjectBrowser.h
@@ -31,9 +31,13 @@ public:
 
 	DECLARE_DELEGATE(FNewProjectScreenRequested)
 
-	SLATE_BEGIN_ARGS(SProjectBrowser) { }
+	SLATE_BEGIN_ARGS(SProjectBrowser) :
+		_HideOpenButton(false)
+	{}
 
-	SLATE_ARGUMENT(FProjectSelectionChanged, OnSelectionChanged);
+		SLATE_ARGUMENT(FProjectSelectionChanged, OnSelectionChanged);
+
+		SLATE_ARGUMENT(bool, HideOpenButton)
 
 	SLATE_END_ARGS()
 
diff --git a/Engine/Source/Editor/GameProjectGeneration/Private/TemplateItem.h b/Engine/Source/Editor/GameProjectGeneration/Private/TemplateItem.h
--- a/Engine/Source/Editor/GameProjectGeneration/Private/TemplateItem.h
+++ b/Engine/Source/Editor/GameProjectGeneration/Private/TemplateItem.h
@@ -28,4 +28,5 @@ struct FTemplateItem
 	TArray<ETemplateSetting> HiddenSettings;
 
 	bool		bIsEnterprise;
+	bool		bIsBlankTemplate;
 };
diff --git a/Engine/Source/Editor/GameProjectGeneration/Public/GameProjectUtils.h b/Engine/Source/Editor/GameProjectGeneration/Public/GameProjectUtils.h
--- a/Engine/Source/Editor/GameProjectGeneration/Public/GameProjectUtils.h
+++ b/Engine/Source/Editor/GameProjectGeneration/Public/GameProjectUtils.h
@@ -23,6 +23,7 @@ struct FProjectInformation
 
 	bool bShouldGenerateCode = false;
 	bool bCopyStarterContent = false;
+	bool bIsBlankTemplate = false;
 	
 	bool bForceExtendedLuminanceRange = false; // See "r.DefaultFeature.AutoExposure.ExtendDefaultLuminanceRange"
 	bool bEnableXR = false;
@@ -301,9 +302,6 @@ private:
 	 */
 	static bool AddSharedContentToProject(const FProjectInformation &InProjectInfo, TArray<FString> &CreatedFiles, FText& OutFailReason);
 
-	/** Returns list of starter content files */
-	static void GetStarterContentFiles(TArray<FString>& OutFilenames);
-
 	/** Returns the template defs ini filename */
 	static FString GetTemplateDefsFilename();
 
@@ -508,7 +506,6 @@ private:
 
 	static TWeakPtr<SNotificationItem> UpdateGameProjectNotification;
 	static TWeakPtr<SNotificationItem> WarningProjectNameNotification;
-	static FString DefaultFeaturePackExtension;
 
 	// Whether we should use AudioMixer for all platforms:
 	static bool bUseAudioMixerForAllPlatforms;
diff --git a/Templates/TP_Blank/Config/TemplateDefs.ini b/Templates/TP_Blank/Config/TemplateDefs.ini
--- a/Templates/TP_Blank/Config/TemplateDefs.ini
+++ b/Templates/TP_Blank/Config/TemplateDefs.ini
@@ -20,6 +20,9 @@ SortKey="_1"
 ; These should contain files only of concern to the template itself and not desired in the user's project.
 FoldersToIgnore=Media
 
+; Is this a blank template? Blank templates do not specify a default map.
+bIsBlank=True
+
 ; The files to ignore copying, relative to the template root folder. 
 ; %TEMPLATENAME% will be replaced by the template's folder name when resolving these. 
 ; Filename comparisons are case-insensitive. 
@@ -82,13 +85,13 @@ ReplacementsInFiles=(Extensions=("cpp","h","ini","cs"), From="%TEMPLATENAME%", T
 ; HardwareTarget - Choose between Mobile and Desktop hardware targets.
 ; GraphicsPreset - Choose between Maximum and Scalable graphics.
 ; StarterContent - Choose whether to add starter content to the project.
-; VR - Choose whether to enable Virtual Reality support for the project.
+; XR - Choose whether to enable Mixed Reality support for the project.
 ;HiddenSettings=All
 ;HiddenSettings=Languages
 ;HiddenSettings=HardwareTarget
 ;HiddenSettings=GraphicsPreset 
 ;HiddenSettings=StarterContent 
-;HiddenSettings=VR 
+;HiddenSettings=XR 
 
 ; The class types that this project uses.
 ; This is only for display in the template description, and does not affect project creation.
diff --git a/Templates/TP_BlankBP/Config/TemplateDefs.ini b/Templates/TP_BlankBP/Config/TemplateDefs.ini
--- a/Templates/TP_BlankBP/Config/TemplateDefs.ini
+++ b/Templates/TP_BlankBP/Config/TemplateDefs.ini
@@ -2,6 +2,8 @@
 
 SortKey="_1"
 
+bIsBlank=True
+
 Categories=Games
 
 LocalizedDisplayNames=(Language="en", Text="Blank")
