Fixed Blank project not being added to empty categories in the New Project dialog.

diff --git a/Engine/Source/Editor/GameProjectGeneration/Private/SNewProjectWizard.cpp b/Engine/Source/Editor/GameProjectGeneration/Private/SNewProjectWizard.cpp
--- a/Engine/Source/Editor/GameProjectGeneration/Private/SNewProjectWizard.cpp
+++ b/Engine/Source/Editor/GameProjectGeneration/Private/SNewProjectWizard.cpp
@@ -922,13 +922,6 @@ TMap<FName, TArray<TSharedPtr<FTemplateItem>> >& SNewProjectWizard::FindTemplate
 		}
 	}
 
-
-	if (Templates.Num() == 0)
-	{
-		static const FName DefaultCategory("Default");
-		Templates.Add(DefaultCategory);
-	}
-
 	// add blank template to empty categories
 	{
 		TSharedPtr<FTemplateItem> BlankTemplate = MakeShareable(new FTemplateItem());
