Added a message when there are no Recent Projects available, because it was unclear what was wrong when there were none showing.

Changed project creation wizard titles to be clearer.

diff --git a/Engine/Source/Editor/GameProjectGeneration/Private/SGameProjectDialog.cpp b/Engine/Source/Editor/GameProjectGeneration/Private/SGameProjectDialog.cpp
--- a/Engine/Source/Editor/GameProjectGeneration/Private/SGameProjectDialog.cpp
+++ b/Engine/Source/Editor/GameProjectGeneration/Private/SGameProjectDialog.cpp
@@ -2,10 +2,12 @@
 
 #include "SGameProjectDialog.h"
 #include "Misc/MessageDialog.h"
-#include "Widgets/SBoxPanel.h"
-#include "Widgets/SOverlay.h"
+
 #include "Styling/CoreStyle.h"
 #include "SlateOptMacros.h"
+#include "Widgets/SBoxPanel.h"
+#include "Widgets/SOverlay.h"
+#include "Widgets/SToolTip.h"
 #include "Widgets/Input/SButton.h"
 #include "Widgets/Input/SHyperlink.h"
 #include "Widgets/Images/SImage.h"
@@ -14,11 +16,11 @@
 #include "Widgets/Layout/SSeparator.h"
 #include "Widgets/Layout/SUniformGridPanel.h"
 #include "Widgets/Layout/SWidgetSwitcher.h"
+#include "Widgets/Text/SRichTextBlock.h"
 #include "EditorStyleSet.h"
 #include "GameProjectUtils.h"
 #include "SProjectBrowser.h"
 #include "SNewProjectWizard.h"
-#include "Widgets/SToolTip.h"
 #include "IDocumentation.h"
 #include "Framework/MultiBox/MultiBoxBuilder.h"
 #include "Internationalization/BreakIterator.h"
@@ -49,9 +51,9 @@ namespace GameProjectDialogDefs
 	constexpr int32 TemplateListPageIndex = 2;
 	constexpr int32 ProjectSettingsPageIndex = 3;
 	
-	static const FText LandingPageTitle = LOCTEXT("ProjectDialog_SelectProject", "Select Project");
+	static const FText LandingPageTitle = LOCTEXT("ProjectDialog_SelectOrCreateProject", "Select or Create New Project");
 	static const FText SelectCategoryTitle = LOCTEXT("ProjectDialog_SelectTemplateCategory", "Select Template Category");
-	static const FText ProjectBrowserTitle = LOCTEXT("ProjectDialog_ProjectBrowser", "Project Browser");
+	static const FText ProjectBrowserTitle = LOCTEXT("ProjectDialog_ProjectBrowserTitle", "Open Existing Project");
 	static const FText TemplateListTitle = LOCTEXT("ProjectDialog_SelectTemplate", "Select Template");
 	static const FText ProjectSettingsTitle = LOCTEXT("ProjectDialog_ProjectSettings", "Project Settings");
 }
@@ -172,6 +174,7 @@ TSharedRef<SWidget> SGameProjectDialog::CreateLandingPage()
 			+ SScrollBox::Slot()
 			[
 				SNew(SVerticalBox)
+				
 				+ SVerticalBox::Slot()
 				.AutoHeight()
 				[
@@ -181,10 +184,40 @@ TSharedRef<SWidget> SGameProjectDialog::CreateLandingPage()
 					[
 						SAssignNew(RecentProjectBrowser, SRecentProjectBrowser)
 						.OnSelectionChanged(this, &SGameProjectDialog::OnRecentProjectSelectionChanged)
+						.Visibility(this, &SGameProjectDialog::GetRecentProjectBrowserVisibility)
+					]
+					+ SOverlay::Slot()
+					.Padding(8, 0)
+					[
+						SNew(SVerticalBox)
+						.Visibility(this, &SGameProjectDialog::GetNoRecentProjectsLabelVisibility)
+						+ SVerticalBox::Slot()
+						.AutoHeight()
+						.Padding(0, 4, 0, 0)
+						[
+							SNew(STextBlock)
+							.TextStyle(FEditorStyle::Get(), "GameProjectDialog.ProjectNamePathLabels")
+							.Text(LOCTEXT("ProjectDialog_Recent", "Recent Projects"))
+						]
+						+ SVerticalBox::Slot()
+						.AutoHeight()
+						.Padding(0, 2, 0, 4)
+						[
+							SNew(SSeparator)
+						]
+						+ SVerticalBox::Slot()
+						.AutoHeight()
+						.Padding(0, 8, 0, 4)
+						[
+							SNew(SRichTextBlock)
+							.DecoratorStyleSet(&FEditorStyle::Get())
+							.Text(LOCTEXT("ProjectDialog_NoRecentProjects", "No recent projects found. Press <RichTextBlock.BoldHighlight>More</> to browse for projects."))
+						]
 					]
 					+ SOverlay::Slot()
 					.HAlign(HAlign_Right)
 					.VAlign(VAlign_Bottom)
+					.Padding(8, 12, 8, 0)
 					[
 						SNew(SButton)
 						.Text(LOCTEXT("ProjectDialog_More", "More"))
@@ -199,7 +232,7 @@ TSharedRef<SWidget> SGameProjectDialog::CreateLandingPage()
 				]
 				+ SVerticalBox::Slot()
 				.AutoHeight()
-				.Padding(0, 4, 0, 2)
+				.Padding(8, 4, 8, 4)
 				[
 					SNew(SSeparator)
 					.Visibility(DialogMode == EMode::Both ? EVisibility::Visible : EVisibility::Collapsed)
@@ -210,7 +243,7 @@ TSharedRef<SWidget> SGameProjectDialog::CreateLandingPage()
 				[
 					SNew(STextBlock)
 					.TextStyle(FEditorStyle::Get(), "GameProjectDialog.ProjectNamePathLabels")
-					.Text(LOCTEXT("ProjectDialog_Categories", "Categories"))
+					.Text(LOCTEXT("ProjectDialog_NewProjectCategories", "New Project Categories"))
 					.Visibility(DialogMode == EMode::Both ? EVisibility::Visible : EVisibility::Collapsed)
 				]
 				+ SVerticalBox::Slot()
@@ -520,4 +553,24 @@ FText SGameProjectDialog::GetPageTitle(int32 PageIndex) const
 	return FText::GetEmpty();
 }
 
-#undef LOCTEXT_NAMESPACE
+EVisibility SGameProjectDialog::GetRecentProjectBrowserVisibility() const
+{
+	if (!RecentProjectBrowser.IsValid())
+	{
+		return EVisibility::Collapsed;
+	}
+
+	return RecentProjectBrowser->HasProjects() ? EVisibility::Visible : EVisibility::Collapsed;
+}
+
+EVisibility SGameProjectDialog::GetNoRecentProjectsLabelVisibility() const
+{
+	if (!RecentProjectBrowser.IsValid())
+	{
+		return EVisibility::Collapsed;
+	}
+
+	return RecentProjectBrowser->HasProjects() ? EVisibility::Collapsed : EVisibility::Visible;
+}
+
+#undef LOCTEXT_NAMESPACE
\ No newline at end of file
diff --git a/Engine/Source/Editor/GameProjectGeneration/Private/SGameProjectDialog.h b/Engine/Source/Editor/GameProjectGeneration/Private/SGameProjectDialog.h
--- a/Engine/Source/Editor/GameProjectGeneration/Private/SGameProjectDialog.h
+++ b/Engine/Source/Editor/GameProjectGeneration/Private/SGameProjectDialog.h
@@ -84,6 +84,9 @@ private:
 
 	void OnEnterSettingsPage();
 
+	EVisibility GetRecentProjectBrowserVisibility() const;
+	EVisibility GetNoRecentProjectsLabelVisibility() const;
+
 	FText GetFinishText() const;
 	FText GetFinishTooltip() const;
 	FText GetPageTitle(int32 Index) const;
