Hidden settings in the new project wizard will now write no values to the config - this puts more burden on the template writer, but shouldn't cause any template settings to be overwritten by hidden settings.

Source: 9731500 (//UE4/Dev-Editor)

diff --git a/Engine/Source/Editor/GameProjectGeneration/Private/GameProjectUtils.cpp b/Engine/Source/Editor/GameProjectGeneration/Private/GameProjectUtils.cpp
--- a/Engine/Source/Editor/GameProjectGeneration/Private/GameProjectUtils.cpp
+++ b/Engine/Source/Editor/GameProjectGeneration/Private/GameProjectUtils.cpp
@@ -166,7 +166,8 @@ namespace
 			TEXT("OculusVR"),
 			TEXT("SteamVR") };
 
-		if (!InProjectInfo.bEnableXR)
+		if (InProjectInfo.bEnableXR.IsSet() && 
+			InProjectInfo.bEnableXR.GetValue() == false)
 		{
 			for (const FString& Plugin : XRPlugins)
 			{
@@ -187,7 +188,8 @@ namespace
 	{
 		FString RaytracingConfig;
 
-		if (InProjectInfo.bEnableRaytracing)
+		if (InProjectInfo.bEnableRaytracing.IsSet() && 
+			InProjectInfo.bEnableRaytracing.GetValue() == true)
 		{
 			RaytracingConfig += LINE_TERMINATOR;
 			RaytracingConfig += TEXT("[/Script/WindowsTargetPlatform.WindowsTargetSettings]") LINE_TERMINATOR;
@@ -868,10 +870,24 @@ bool GameProjectUtils::CreateProject(const FProjectInformation& InProjectInfo, F
 		EventAttributes.Add(FAnalyticsEventAttribute(TEXT("ProjectType"), InProjectInfo.bShouldGenerateCode ? TEXT("C++ Code") : TEXT("Content Only")));
 		EventAttributes.Add(FAnalyticsEventAttribute(TEXT("Outcome"), bProjectCreationSuccessful ? TEXT("Successful") : TEXT("Failed")));
 
-		UEnum* Enum = StaticEnum<EHardwareClass::Type>();
-		EventAttributes.Add(FAnalyticsEventAttribute(TEXT("HardwareClass"), Enum ? Enum->GetNameStringByValue(InProjectInfo.TargetedHardware) : FString()));
-		Enum = StaticEnum<EGraphicsPreset::Type>();
-		EventAttributes.Add(FAnalyticsEventAttribute(TEXT("GraphicsPreset"), Enum ? Enum->GetNameStringByValue(InProjectInfo.DefaultGraphicsPerformance) : FString()));
+		if (InProjectInfo.TargetedHardware.IsSet())
+		{
+			UEnum* HardwareClassEnum = StaticEnum<EHardwareClass::Type>();
+			if (HardwareClassEnum != nullptr)
+			{
+				EventAttributes.Add(FAnalyticsEventAttribute(TEXT("HardwareClass"), HardwareClassEnum->GetNameStringByValue(InProjectInfo.TargetedHardware.GetValue())));
+			}
+		}
+
+		if (InProjectInfo.DefaultGraphicsPerformance.IsSet())
+		{
+			UEnum* GraphicsPresetEnum = StaticEnum<EGraphicsPreset::Type>();
+			if (GraphicsPresetEnum != nullptr)
+			{
+				EventAttributes.Add(FAnalyticsEventAttribute(TEXT("GraphicsPreset"), GraphicsPresetEnum->GetNameStringByValue(InProjectInfo.DefaultGraphicsPerformance.GetValue())));
+			}
+		}
+
 		EventAttributes.Add(FAnalyticsEventAttribute(TEXT("StarterContent"), InProjectInfo.bCopyStarterContent ? TEXT("Yes") : TEXT("No")));
 		
 		FEngineAnalytics::GetProvider().RecordEvent( TEXT( "Editor.NewProject.ProjectCreated" ), EventAttributes );
@@ -1970,18 +1986,42 @@ void GameProjectUtils::DeleteCreatedFiles(const FString& RootFolder, const TArra
 
 FString GameProjectUtils::GetHardwareConfigString(const FProjectInformation& InProjectInfo)
 {
-	FString HardwareTargeting;
+	FString TargetHardwareString;
+	if (InProjectInfo.TargetedHardware.IsSet())
+	{
+		UEnum* HardwareClassEnum = StaticEnum<EHardwareClass::Type>();
+		if (HardwareClassEnum != nullptr)
+		{
+			HardwareClassEnum->GetValueAsString(InProjectInfo.TargetedHardware.GetValue(), /*out*/ TargetHardwareString);
+		}
+	}
 
-	FString TargetHardwareAsString;
-	UEnum::GetValueAsString(TEXT("/Script/HardwareTargeting.EHardwareClass"), InProjectInfo.TargetedHardware, /*out*/ TargetHardwareAsString);
+	FString GraphicsPresetString;
+	if (InProjectInfo.DefaultGraphicsPerformance.IsSet())
+	{
+		UEnum* GraphicsPresetEnum = StaticEnum<EGraphicsPreset::Type>();
+		if (GraphicsPresetEnum != nullptr)
+		{
+			GraphicsPresetEnum->GetValueAsString(InProjectInfo.DefaultGraphicsPerformance.GetValue(), /*out*/ GraphicsPresetString);
+		}
+	}
 
-	FString GraphicsPresetAsString;
-	UEnum::GetValueAsString(TEXT("/Script/HardwareTargeting.EGraphicsPreset"), InProjectInfo.DefaultGraphicsPerformance, /*out*/ GraphicsPresetAsString);
+	FString HardwareTargeting;
+	if (!TargetHardwareString.IsEmpty() || !GraphicsPresetString.IsEmpty())
+	{
+		HardwareTargeting += LINE_TERMINATOR;
+		HardwareTargeting += TEXT("[/Script/HardwareTargeting.HardwareTargetingSettings]") LINE_TERMINATOR;
+	}
+
+	if (!TargetHardwareString.IsEmpty())
+	{
+		HardwareTargeting += FString::Printf(TEXT("TargetedHardwareClass=%s") LINE_TERMINATOR, *TargetHardwareString);
+	}
 
-	HardwareTargeting += LINE_TERMINATOR;
-	HardwareTargeting += TEXT("[/Script/HardwareTargeting.HardwareTargetingSettings]") LINE_TERMINATOR;
-	HardwareTargeting += FString::Printf(TEXT("TargetedHardwareClass=%s") LINE_TERMINATOR, *TargetHardwareAsString);
-	HardwareTargeting += FString::Printf(TEXT("DefaultGraphicsPerformance=%s") LINE_TERMINATOR, *GraphicsPresetAsString);
+	if (!GraphicsPresetString.IsEmpty())
+	{
+		HardwareTargeting += FString::Printf(TEXT("DefaultGraphicsPerformance=%s") LINE_TERMINATOR, *GraphicsPresetString);
+	}
 
 	return MoveTemp(HardwareTargeting);
 }
diff --git a/Engine/Source/Editor/GameProjectGeneration/Private/SNewProjectWizard.cpp b/Engine/Source/Editor/GameProjectGeneration/Private/SNewProjectWizard.cpp
--- a/Engine/Source/Editor/GameProjectGeneration/Private/SNewProjectWizard.cpp
+++ b/Engine/Source/Editor/GameProjectGeneration/Private/SNewProjectWizard.cpp
@@ -303,8 +303,6 @@ private:
 BEGIN_SLATE_FUNCTION_BUILD_OPTIMIZATION
 void SNewProjectWizard::Construct( const FArguments& InArgs )
 {
-	bProjectSettingsHidden = false;
-
 	bLastGlobalValidityCheckSuccessful = true;
 	bLastNameAndLocationValidityCheckSuccessful = true;
 	bCopyStarterContent = GEditor ? GetDefault<UEditorSettings>()->bCopyStarterContentPreference : true;
@@ -1103,19 +1101,40 @@ bool SNewProjectWizard::CreateProject(const FString& ProjectFile)
 
 	FText FailReason, FailLog;
 
+	const TArray<ETemplateSetting>& HiddenSettings = GetSelectedTemplateProperty(&FTemplateItem::HiddenSettings);
+
 	FProjectInformation ProjectInfo;
 	ProjectInfo.ProjectFilename = ProjectFile;
 	ProjectInfo.bShouldGenerateCode = bShouldGenerateCode;
 	ProjectInfo.bCopyStarterContent = bCopyStarterContent;
 	ProjectInfo.TemplateFile = bShouldGenerateCode ? SelectedTemplate->CodeProjectFile : SelectedTemplate->BlueprintProjectFile;
 	ProjectInfo.TemplateCategory = ActiveCategory;
-	ProjectInfo.TargetedHardware = SelectedHardwareClassTarget;
-	ProjectInfo.DefaultGraphicsPerformance = SelectedGraphicsPreset;
-	ProjectInfo.bForceExtendedLuminanceRange = ProjectInfo.TemplateFile.IsEmpty();
-	ProjectInfo.bEnableXR = bEnableXR;
-	ProjectInfo.bEnableRaytracing = bEnableRaytracing;
 	ProjectInfo.bIsEnterpriseProject = SelectedTemplate->bIsEnterprise;
 	ProjectInfo.bIsBlankTemplate = SelectedTemplate->bIsBlankTemplate;
+	ProjectInfo.bForceExtendedLuminanceRange = SelectedTemplate->bIsBlankTemplate;
+
+	if (!HiddenSettings.Contains(ETemplateSetting::All))
+	{
+		if (!HiddenSettings.Contains(ETemplateSetting::HardwareTarget))
+		{
+			ProjectInfo.TargetedHardware = SelectedHardwareClassTarget;
+		}
+
+		if (!HiddenSettings.Contains(ETemplateSetting::GraphicsPreset))
+		{
+			ProjectInfo.DefaultGraphicsPerformance = SelectedGraphicsPreset;
+		}
+
+		if (!HiddenSettings.Contains(ETemplateSetting::XR))
+		{
+			ProjectInfo.bEnableXR = bEnableXR;
+		}
+
+		if (!HiddenSettings.Contains(ETemplateSetting::Raytracing))
+		{
+			ProjectInfo.bEnableRaytracing = bEnableRaytracing;
+		}
+	}
 
 	if (!GameProjectUtils::CreateProject(ProjectInfo, FailReason, FailLog))
 	{
diff --git a/Engine/Source/Editor/GameProjectGeneration/Private/SNewProjectWizard.h b/Engine/Source/Editor/GameProjectGeneration/Private/SNewProjectWizard.h
--- a/Engine/Source/Editor/GameProjectGeneration/Private/SNewProjectWizard.h
+++ b/Engine/Source/Editor/GameProjectGeneration/Private/SNewProjectWizard.h
@@ -253,6 +253,4 @@ private:
 	/** Names for pages */
 	static FName TemplatePageName;
 	static FName NameAndLocationPageName;
-
-	bool bProjectSettingsHidden;
 };
diff --git a/Engine/Source/Editor/GameProjectGeneration/Public/GameProjectUtils.h b/Engine/Source/Editor/GameProjectGeneration/Public/GameProjectUtils.h
--- a/Engine/Source/Editor/GameProjectGeneration/Public/GameProjectUtils.h
+++ b/Engine/Source/Editor/GameProjectGeneration/Public/GameProjectUtils.h
@@ -24,14 +24,17 @@ struct FProjectInformation
 	bool bShouldGenerateCode = false;
 	bool bCopyStarterContent = false;
 	bool bIsBlankTemplate = false;
-	
-	bool bForceExtendedLuminanceRange = false; // See "r.DefaultFeature.AutoExposure.ExtendDefaultLuminanceRange"
-	bool bEnableXR = false;
 	bool bIsEnterpriseProject = false;
-	bool bEnableRaytracing = false;
+	bool bForceExtendedLuminanceRange; // See "r.DefaultFeature.AutoExposure.ExtendDefaultLuminanceRange"
 
-	EHardwareClass::Type TargetedHardware = EHardwareClass::Desktop;
-	EGraphicsPreset::Type DefaultGraphicsPerformance = EGraphicsPreset::Maximum;
+	// These are all optional, because there is an additional state introduced by hiding the setting in the template.
+	// In this case, the template author has chosen not to give the user a choice,
+	// so we must assume that the template already controls these settings explicitly.
+	TOptional<bool> bEnableXR;
+	TOptional<bool> bEnableRaytracing;
+
+	TOptional<EHardwareClass::Type> TargetedHardware;
+	TOptional<EGraphicsPreset::Type> DefaultGraphicsPerformance;
 };
 
 DECLARE_DELEGATE_RetVal_OneParam(bool, FProjectDescriptorModifier, FProjectDescriptor&);
