We now always create an empty Content folder on project creation - this is to support template-based blank projects, which otherwise would not populate the Content folder at all, causing the Content folder directory watcher to not initialize on startup.

diff --git a/Engine/Source/Editor/GameProjectGeneration/Private/GameProjectUtils.cpp b/Engine/Source/Editor/GameProjectGeneration/Private/GameProjectUtils.cpp
--- a/Engine/Source/Editor/GameProjectGeneration/Private/GameProjectUtils.cpp
+++ b/Engine/Source/Editor/GameProjectGeneration/Private/GameProjectUtils.cpp
@@ -1494,6 +1494,16 @@ bool GameProjectUtils::CreateProjectFromTemplate(const FProjectInformation& InPr
 
 	SlowTask.EnterProgressFrame();
 
+	// create a Content folder in the destination directory
+	const FString DestContentFolder = DestFolder / TEXT("Content");
+	if (!IFileManager::Get().MakeDirectory(*DestContentFolder, true))
+	{
+		FFormatNamedArguments FailArgs;
+		FailArgs.Add(TEXT("DestContentFolder"), FText::FromString(DestContentFolder));
+		OutFailReason = FText::Format(LOCTEXT("FailedToCreateDirectory", "Failed to create directory \"{DestContentFolder}\"."), FailArgs);
+		return false;
+	}
+
 	// Add a rule to replace the build settings version with the appropriate number
 	FTemplateReplacement& DefaultBuildSettingsRepl = TemplateDefs->ReplacementsInFiles.AddZeroed_GetRef();
 	DefaultBuildSettingsRepl.Extensions.Add("cs");
