Fix failures to add code to a content project once a custom executable has been built.

diff --git a/Engine/Source/Editor/GameProjectGeneration/Private/GameProjectUtils.cpp b/Engine/Source/Editor/GameProjectGeneration/Private/GameProjectUtils.cpp
--- a/Engine/Source/Editor/GameProjectGeneration/Private/GameProjectUtils.cpp
+++ b/Engine/Source/Editor/GameProjectGeneration/Private/GameProjectUtils.cpp
@@ -3873,6 +3873,9 @@ GameProjectUtils::EAddCodeToProjectResult GameProjectUtils::AddCodeToProject_Int
 	const bool bProjectHadCodeFiles = ProjectHasCodeFiles();
 	if (!bProjectHadCodeFiles)
 	{
+		// Delete any generated intermediate code files. This ensures that blueprint projects with custom build settings can be converted to code projects without causing errors.
+		IFileManager::Get().DeleteDirectory(*(FPaths::ProjectIntermediateDir() / TEXT("Source")), false, true);
+
 		// We always add the basic source code to the root directory, not the potential sub-directory provided by NewClassPath
 		const FString SourceDir = FPaths::GameSourceDir().LeftChop(1); // Trim the trailing /
 
