Moved warning about missing compiler to above the project location entry box.

#rnx

diff --git a/Engine/Source/Editor/GameProjectGeneration/Private/SNewProjectWizard.cpp b/Engine/Source/Editor/GameProjectGeneration/Private/SNewProjectWizard.cpp
--- a/Engine/Source/Editor/GameProjectGeneration/Private/SNewProjectWizard.cpp
+++ b/Engine/Source/Editor/GameProjectGeneration/Private/SNewProjectWizard.cpp
@@ -721,7 +721,7 @@ EVisibility SNewProjectWizard::GetNameAndLocationErrorLabelVisibility() const
 
 FText SNewProjectWizard::GetNameAndLocationErrorLabelText() const
 {
-	if ( !bLastNameAndLocationValidityCheckSuccessful )
+	if (!bLastNameAndLocationValidityCheckSuccessful)
 	{
 		return LastNameAndLocationValidityErrorText;
 	}
@@ -1363,12 +1363,12 @@ TSharedRef<SWidget> SNewProjectWizard::CreateProjectSettingsPage()
 	+ SOverlay::Slot()
 	.HAlign(HAlign_Left)
 	.VAlign(VAlign_Bottom)
-	.Padding( 8 )
+	.Padding(0, 0, 0, 82) // manually sized to be above the project location box
 	[
 		SNew(SBorder)
 		.Visibility(this, &SNewProjectWizard::GetGlobalErrorLabelVisibility)
 		.BorderImage(FEditorStyle::GetBrush("GameProjectDialog.ErrorLabelBorder"))
-		.Padding( 8 )
+		.Padding(4)
 		[
 			SNew(SHorizontalBox)
 						
@@ -1421,12 +1421,11 @@ TSharedRef<SWidget> SNewProjectWizard::CreateProjectSettingsPage()
 	+SOverlay::Slot()
 	.HAlign(HAlign_Left)
 	.VAlign(VAlign_Bottom)
-	.Padding( 8 )
 	[
 		SNew(SBorder)
 		.BorderImage(FEditorStyle::GetBrush("GameProjectDialog.ErrorLabelBorder"))
 		.Visibility(this, &SNewProjectWizard::GetNameAndLocationErrorLabelVisibility)
-		.Padding(8)
+		.Padding(4)
 		[
 			SNew(SHorizontalBox)
 					
@@ -1441,10 +1440,9 @@ TSharedRef<SWidget> SNewProjectWizard::CreateProjectSettingsPage()
 
 			+SHorizontalBox::Slot()
 			.VAlign(VAlign_Center)
-			.AutoWidth()
+			.FillWidth(1.0f)
 			[
 				SNew(STextBlock)
-				.AutoWrapText(true)
 				.Text(this, &SNewProjectWizard::GetNameAndLocationErrorLabelText)
 				.TextStyle(FEditorStyle::Get(), "GameProjectDialog.ErrorLabelFont")
 			]
