The New Project dialog's XR setting has been removed as it was not functioning with binary and packaged builds.

diff --git a/Engine/Source/Editor/GameProjectGeneration/Private/GameProjectUtils.cpp b/Engine/Source/Editor/GameProjectGeneration/Private/GameProjectUtils.cpp
--- a/Engine/Source/Editor/GameProjectGeneration/Private/GameProjectUtils.cpp
+++ b/Engine/Source/Editor/GameProjectGeneration/Private/GameProjectUtils.cpp
@@ -157,6 +157,7 @@ namespace
 		return MoveTemp(DefaultProjectSettings);
 	}
 
+	// @todo: This is currently not called from anywhere as this approach does not work for binary builds.
 	/** Set the state of XR plugins in OutProject based on the flags in InProjectInfo. */
 	void SetXRPluginStates(const FProjectInformation& InProjectInfo, FProjectDescriptor& OutProject)
 	{
@@ -1428,8 +1429,6 @@ bool GameProjectUtils::GenerateProjectFromScratch(const FProjectInformation& InP
 
 		Project.bIsEnterpriseProject = InProjectInfo.bIsEnterpriseProject;
 
-		SetXRPluginStates(InProjectInfo, Project);
-
 		// Try to save it
 		FText LocalFailReason;
 		if(!Project.Save(InProjectInfo.ProjectFilename, LocalFailReason))
@@ -1841,9 +1840,6 @@ bool GameProjectUtils::CreateProjectFromTemplate(const FProjectInformation& InPr
 			ModuleInfo.Name = FName(*ModuleInfo.Name.ToString().Replace(*BaseSourceName, *BaseNewName));
 		}
 
-		// Disable XR plugins if desired
-		SetXRPluginStates(InProjectInfo, Project);
-
 		// Save it to disk
 		if(!Project.Save(InProjectInfo.ProjectFilename, OutFailReason))
 		{
diff --git a/Engine/Source/Editor/GameProjectGeneration/Private/SNewProjectWizard.cpp b/Engine/Source/Editor/GameProjectGeneration/Private/SNewProjectWizard.cpp
--- a/Engine/Source/Editor/GameProjectGeneration/Private/SNewProjectWizard.cpp
+++ b/Engine/Source/Editor/GameProjectGeneration/Private/SNewProjectWizard.cpp
@@ -1625,6 +1625,7 @@ TSharedRef<SWidget> SNewProjectWizard::MakeProjectSettingsOptionsBox()
 		AddToProjectSettingsGrid(GridPanel, Enum, Description, CurrentSlot);
 	}
 
+#if 0 // @todo: XR settings cannot be shown at the moment as the setting causes issues with binary builds.
 	if (!HiddenSettings.Contains(ETemplateSetting::XR))
 	{
 		TArray<SDecoratedEnumCombo<int32>::FComboOption> VirtualRealityOptions;
@@ -1649,6 +1650,7 @@ TSharedRef<SWidget> SNewProjectWizard::MakeProjectSettingsOptionsBox()
 
 		AddToProjectSettingsGrid(GridPanel, Enum, Description, CurrentSlot);
 	}
+#endif 
 
 	if (!HiddenSettings.Contains(ETemplateSetting::Raytracing))
 	{
