The warning in the New Project wizard about mobile starter content increasing package size now correctly displays on top of the enum it's placed on top of.

diff --git a/Engine/Source/Editor/GameProjectGeneration/Private/SNewProjectWizard.cpp b/Engine/Source/Editor/GameProjectGeneration/Private/SNewProjectWizard.cpp
--- a/Engine/Source/Editor/GameProjectGeneration/Private/SNewProjectWizard.cpp
+++ b/Engine/Source/Editor/GameProjectGeneration/Private/SNewProjectWizard.cpp
@@ -1151,7 +1151,7 @@ bool SNewProjectWizard::CreateProject(const FString& ProjectFile)
 		CreatedProjectPath.AppendChar('/');
 	}
 
-	auto* Settings = GetMutableDefault<UEditorSettings>();
+	UEditorSettings* Settings = GetMutableDefault<UEditorSettings>();
 	Settings->CreatedProjectPaths.Remove(CreatedProjectPath);
 	Settings->CreatedProjectPaths.Insert(CreatedProjectPath, 0);
 	Settings->bCopyStarterContentPreference = bCopyStarterContent;
@@ -1597,6 +1597,13 @@ TSharedRef<SWidget> SNewProjectWizard::MakeProjectSettingsOptionsBox()
 			bIsStarterAvailable));
 
 		TSharedRef<SWidget> Enum = SNew(SOverlay)
+			+ SOverlay::Slot()
+			[
+				SNew(SDecoratedEnumCombo<int32>, MoveTemp(StarterContentOptions))
+				.SelectedEnum(this, &SNewProjectWizard::GetCopyStarterContentIndex)
+				.OnEnumChanged(this, &SNewProjectWizard::OnSetCopyStarterContent)
+				.Orientation(Orient_Vertical)
+			]
 			+ SOverlay::Slot()
 			.HAlign(HAlign_Right)
 			.VAlign(VAlign_Top)
@@ -1607,13 +1614,6 @@ TSharedRef<SWidget> SNewProjectWizard::MakeProjectSettingsOptionsBox()
 				.Image(FEditorStyle::GetBrush("Icons.Warning"))
 				.ToolTipText(this, &SNewProjectWizard::GetStarterContentWarningTooltip)
 				.Visibility(this, &SNewProjectWizard::GetStarterContentWarningVisibility)
-			]
-			+ SOverlay::Slot()
-			[
-				SNew(SDecoratedEnumCombo<int32>, MoveTemp(StarterContentOptions))
-				.SelectedEnum(this, &SNewProjectWizard::GetCopyStarterContentIndex)
-				.OnEnumChanged(this, &SNewProjectWizard::OnSetCopyStarterContent)
-				.Orientation(Orient_Vertical)
 			];
 
 		TSharedRef<SRichTextBlock> Description =
