Changed path for automated test to have a full path name instead of a relative one.
#rnx

diff --git a/Engine/Source/Editor/GameProjectGeneration/Private/Tests/GameProjectAutomationTests.cpp b/Engine/Source/Editor/GameProjectGeneration/Private/Tests/GameProjectAutomationTests.cpp
--- a/Engine/Source/Editor/GameProjectGeneration/Private/Tests/GameProjectAutomationTests.cpp
+++ b/Engine/Source/Editor/GameProjectGeneration/Private/Tests/GameProjectAutomationTests.cpp
@@ -309,7 +309,8 @@ bool FBuildPromotionNewProjectMapTest::RunTest(const FString& Parameters)
 	}
 
 	// Save the map
-	FEditorFileUtils::SaveLevel(CurrentLevel, TEXT("/Game/Maps/NewProjectTest"));
+    FString PathToSave = FPaths::ConvertRelativePathToFull(FPaths::ProjectContentDir() + TEXT("Maps/NewProjectTest.umap"));
+	FEditorFileUtils::SaveLevel(CurrentLevel, PathToSave);
 	UE_LOG(LogGameProjectGenerationTests, Display, TEXT("Saved map"));
 
 	return true;
