New Project config generation now places config values into the right section in the INI without duplicating the section.

This is necessary because changing Project Settings would add settings to the first section of a given name and not modify any further instances of that section. As a result, it was possible to get conflicting values in the INI, effectively preventing certain settings from ever being changed without hand editing the INI file in question.

diff --git a/Engine/Source/Editor/GameProjectGeneration/Private/GameProjectUtils.cpp b/Engine/Source/Editor/GameProjectGeneration/Private/GameProjectUtils.cpp
--- a/Engine/Source/Editor/GameProjectGeneration/Private/GameProjectUtils.cpp
+++ b/Engine/Source/Editor/GameProjectGeneration/Private/GameProjectUtils.cpp
@@ -185,49 +185,47 @@ namespace
 	}
 
 	/** Get the configuration values for raytracing if enabled. */
-	FString GetRaytracingConfigString(const FProjectInformation& InProjectInfo)
+	void AddRaytracingConfigValues(const FProjectInformation& InProjectInfo, TArray<FTemplateConfigValue>& ConfigValues)
 	{
-		FString RaytracingConfig;
-
 		if (InProjectInfo.bEnableRaytracing.IsSet() && 
 			InProjectInfo.bEnableRaytracing.GetValue() == true)
 		{
-			RaytracingConfig += LINE_TERMINATOR;
-			RaytracingConfig += TEXT("[/Script/WindowsTargetPlatform.WindowsTargetSettings]") LINE_TERMINATOR;
-			RaytracingConfig += TEXT("DefaultGraphicsRHI=DefaultGraphicsRHI_DX12") LINE_TERMINATOR;
-			RaytracingConfig += LINE_TERMINATOR;
-			RaytracingConfig += TEXT("[/Script/Engine.RendererSettings]") LINE_TERMINATOR;
-			RaytracingConfig += TEXT("r.SkinCache.CompileShaders=True") LINE_TERMINATOR;
-			RaytracingConfig += TEXT("r.RayTracing=True") LINE_TERMINATOR;
-		}
+			ConfigValues.Emplace(TEXT("DefaultEngine.ini"),
+				TEXT("/Script/WindowsTargetPlatform.WindowsTargetSettings"),
+				TEXT("DefaultGraphicsRHI"),
+				TEXT("DefaultGraphicsRHI_DX12"),
+				true /* ShouldReplaceExistingValue */);
 
-		return MoveTemp(RaytracingConfig);
+			ConfigValues.Emplace(TEXT("DefaultEngine.ini"),
+				TEXT("/Script/Engine.RendererSettings"),
+				TEXT("r.SkinCache.CompileShaders"),
+				TEXT("True"),
+				true /* ShouldReplaceExistingValue */);
+
+			ConfigValues.Emplace(TEXT("DefaultEngine.ini"),
+				TEXT("/Script/Engine.RendererSettings"),
+				TEXT("r.RayTracing"),
+				TEXT("True"),
+				true /* ShouldReplaceExistingValue */);
+		}
 	}
 
-	FString GetDefaultMapConfigString(const FProjectInformation& InProjectInfo)
+	void AddDefaultMapConfigValues(const FProjectInformation& InProjectInfo, TArray<FTemplateConfigValue>& ConfigValues)
 	{
-		FString DefaultMapConfig;
-
 		if (InProjectInfo.bIsBlankTemplate &&
 			InProjectInfo.bCopyStarterContent &&
 			GameProjectUtils::IsStarterContentAvailableForNewProjects())
 		{
-			DefaultMapConfig += LINE_TERMINATOR;
-			DefaultMapConfig += TEXT("[/Script/EngineSettings.GameMapsSettings]") LINE_TERMINATOR;
+			const FString DefaultMap = InProjectInfo.TargetedHardware == EHardwareClass::Mobile ?
+				TEXT("/Game/MobileStarterContent/Maps/Minimal_Default") :
+				TEXT("/Game/StarterContent/Maps/Minimal_Default");
 
-			if (InProjectInfo.TargetedHardware == EHardwareClass::Mobile)
-			{
-				DefaultMapConfig += TEXT("EditorStartupMap=/Game/MobileStarterContent/Maps/Minimal_Default") LINE_TERMINATOR;
-				DefaultMapConfig += TEXT("GameDefaultMap=/Game/MobileStarterContent/Maps/Minimal_Default") LINE_TERMINATOR;
-			}
-			else
-			{
-				DefaultMapConfig += TEXT("EditorStartupMap=/Game/StarterContent/Maps/Minimal_Default") LINE_TERMINATOR;
-				DefaultMapConfig += TEXT("GameDefaultMap=/Game/StarterContent/Maps/Minimal_Default") LINE_TERMINATOR;
-			}
+			ConfigValues.Emplace(TEXT("DefaultEngine.ini"),
+				TEXT("/Script/EngineSettings.GameMapsSettings"),
+				TEXT("EditorStartupMap"),
+				DefaultMap,
+				true /* ShouldReplaceExistingValue */);
 		}
-
-		return MoveTemp(DefaultMapConfig);
 	}
 } // namespace <>
 
@@ -1462,6 +1460,113 @@ bool GameProjectUtils::GenerateProjectFromScratch(const FProjectInformation& InP
 	return true;
 }
 
+static bool SaveConfigValues(const FProjectInformation& InProjectInfo, const TArray<FTemplateConfigValue>& ConfigValues, FText& OutFailReason)
+{
+	const FString ProjectConfigPath = FPaths::GetPath(InProjectInfo.ProjectFilename) / TEXT("Config");
+
+	// Fix all specified config values
+	for (const FTemplateConfigValue& ConfigValue : ConfigValues)
+	{
+		const FString IniFilename = ProjectConfigPath / ConfigValue.ConfigFile;
+		bool bSuccessfullyProcessed = false;
+
+		TArray<FString> FileLines;
+		if (FFileHelper::LoadANSITextFileToStrings(*IniFilename, &IFileManager::Get(), FileLines))
+		{
+			FString FileOutput;
+			const FString TargetSection = ConfigValue.ConfigSection;
+			FString CurSection;
+			bool bFoundTargetKey = false;
+			for (const FString& LineIn : FileLines)
+			{
+				FString Line = LineIn;
+				Line.TrimStartAndEndInline();
+
+				bool bShouldExcludeLineFromOutput = false;
+
+				// If we not yet found the target key parse each line looking for it
+				if (!bFoundTargetKey)
+				{
+					if (Line.Len() == 0)
+					{
+						// Check for an empty line. No work needs to be done on these lines
+					}
+					else if (Line.StartsWith(TEXT(";")))
+					{
+						// Comment lines start with ";". Skip these lines entirely.
+					}
+					else if (Line.StartsWith(TEXT("[")))
+					{
+						// If this is a section line, update the section
+						// If we are entering a new section and we have not yet found our key in the target section, add it to the end of the section
+						if (CurSection == TargetSection)
+						{
+							FileOutput += ConfigValue.ConfigKey + TEXT("=") + ConfigValue.ConfigValue + LINE_TERMINATOR + LINE_TERMINATOR;
+							bFoundTargetKey = true;
+						}
+
+						// Update the current section
+						CurSection = Line.Mid(1, Line.Len() - 2);
+					}
+					// This is possibly an actual key/value pair
+					else if (CurSection == TargetSection)
+					{
+						// Key value pairs contain an equals sign
+						const int32 EqualsIdx = Line.Find(TEXT("="));
+						if (EqualsIdx != INDEX_NONE)
+						{
+							// Determine the key and see if it is the target key
+							const FString Key = Line.Left(EqualsIdx);
+							if (Key == ConfigValue.ConfigKey)
+							{
+								// Found the target key, add it to the output and skip the current line if the target value is supposed to replace
+								FileOutput += ConfigValue.ConfigKey + TEXT("=") + ConfigValue.ConfigValue;
+								bShouldExcludeLineFromOutput = ConfigValue.bShouldReplaceExistingValue;
+								bFoundTargetKey = true;
+							}
+						}
+					}
+				}
+
+				// Unless we replaced the key, add this line to the output
+				if (!bShouldExcludeLineFromOutput)
+				{
+					FileOutput += Line;
+					if (&LineIn != &FileLines.Last())
+					{
+						// Add a line terminator on every line except the last
+						FileOutput += LINE_TERMINATOR;
+					}
+				}
+			}
+
+			// If the key did not exist, add it here
+			if (!bFoundTargetKey)
+			{
+				// If we did not end in the correct section, add the section to the bottom of the file
+				if (CurSection != TargetSection)
+				{
+					FileOutput += LINE_TERMINATOR;
+					FileOutput += FString::Printf(TEXT("[%s]"), *TargetSection) + LINE_TERMINATOR;
+				}
+
+				// Add the key/value here
+				FileOutput += ConfigValue.ConfigKey + TEXT("=") + ConfigValue.ConfigValue + LINE_TERMINATOR;
+			}
+
+			if (!FFileHelper::SaveStringToFile(FileOutput, *IniFilename))
+			{
+				FFormatNamedArguments Args;
+				Args.Add(TEXT("ConfigFile"), FText::FromString(IniFilename));
+				OutFailReason = LOCTEXT("FailedToFixUpConfigFile", "Failed to process file {ConfigFile}.");
+				return false;
+			}
+		}
+	}
+
+	return true;
+}
+
 bool GameProjectUtils::CreateProjectFromTemplate(const FProjectInformation& InProjectInfo, FText& OutFailReason, FText& OutFailLog,TArray<FString>* OutCreatedFiles)
 {
 	FScopedSlowTask SlowTask(10);
@@ -1661,30 +1766,15 @@ bool GameProjectUtils::CreateProjectFromTemplate(const FProjectInformation& InPr
 
 	SlowTask.EnterProgressFrame();
 
-	const FString ProjectConfigPath = DestFolder / TEXT("Config");
-
-	// Write out DefaultEngine.ini settings
-	{
-		const FString DefaultEngineIniFilename = ProjectConfigPath / TEXT("DefaultEngine.ini");
-
-		FString FileContents;
-		// Load the existing file - if it doesn't exist we create it
-		FFileHelper::LoadFileToString(FileContents, *DefaultEngineIniFilename);
-
-		FileContents += GetHardwareConfigString(InProjectInfo);
-
-		FileContents += GetRaytracingConfigString(InProjectInfo);
+	// Fixup specific ini values
+	TArray<FTemplateConfigValue> ConfigValuesToSet;
 
-		FileContents += GetDefaultMapConfigString(InProjectInfo);
+	AddHardwareConfigValues(InProjectInfo, ConfigValuesToSet);
 
-		if (!WriteOutputFile(DefaultEngineIniFilename, FileContents, OutFailReason))
-		{
-			return false;
-		}
-	}
+	AddRaytracingConfigValues(InProjectInfo, ConfigValuesToSet);
 
-	// Fixup specific ini values
-	TArray<FTemplateConfigValue> ConfigValuesToSet;
+	AddDefaultMapConfigValues(InProjectInfo, ConfigValuesToSet);
+	
 	TemplateDefs->AddConfigValues(ConfigValuesToSet, TemplateName, ProjectName, InProjectInfo.bShouldGenerateCode);
 	ConfigValuesToSet.Emplace(TEXT("DefaultGame.ini"), TEXT("/Script/EngineSettings.GeneralProjectSettings"), TEXT("ProjectID"), FGuid::NewGuid().ToString(), /*InShouldReplaceExistingValue=*/true);
 
@@ -1695,110 +1785,9 @@ bool GameProjectUtils::CreateProjectFromTemplate(const FProjectInformation& InPr
 		ConfigValuesToSet.Emplace(TEXT("DefaultEngine.ini"), TEXT("/Script/Engine.Engine"), TEXT("+ActiveClassRedirects"), *ClassRedirectString, /*InShouldReplaceExistingValue=*/false);
 	}
 
-	// Fix all specified config values
-	for ( const FTemplateConfigValue& ConfigValue : ConfigValuesToSet )
+	if (!SaveConfigValues(InProjectInfo, ConfigValuesToSet, OutFailReason))
 	{
-		const FString IniFilename = ProjectConfigPath / ConfigValue.ConfigFile;
-		bool bSuccessfullyProcessed = false;
-
-		TArray<FString> FileLines;
-		if ( FFileHelper::LoadANSITextFileToStrings(*IniFilename, &IFileManager::Get(), FileLines) )
-		{
-			FString FileOutput;
-			const FString TargetSection = ConfigValue.ConfigSection;
-			FString CurSection;
-			bool bFoundTargetKey = false;
-			for ( const FString& LineIn : FileLines )
-			{
-				FString Line = LineIn;
-				Line.TrimStartAndEndInline();
-
-				bool bShouldExcludeLineFromOutput = false;
-
-				// If we not yet found the target key parse each line looking for it
-				if ( !bFoundTargetKey )
-				{
-					// Check for an empty line. No work needs to be done on these lines
-					if ( Line.Len() == 0 )
-					{
-
-					}
-					// Comment lines start with ";". Skip these lines entirely.
-					else if ( Line.StartsWith(TEXT(";")) )
-					{
-
-					}
-					// If this is a section line, update the section
-					else if ( Line.StartsWith(TEXT("[")) )
-					{
-						// If we are entering a new section and we have not yet found our key in the target section, add it to the end of the section
-						if ( CurSection == TargetSection )
-						{
-							FileOutput += ConfigValue.ConfigKey + TEXT("=") + ConfigValue.ConfigValue + LINE_TERMINATOR + LINE_TERMINATOR;
-							bFoundTargetKey = true;
-						}
-
-						// Update the current section
-						CurSection = Line.Mid(1, Line.Len() - 2);
-					}
-					// This is possibly an actual key/value pair
-					else if ( CurSection == TargetSection )
-					{
-						// Key value pairs contain an equals sign
-						const int32 EqualsIdx = Line.Find(TEXT("="));
-						if ( EqualsIdx != INDEX_NONE )
-						{
-							// Determine the key and see if it is the target key
-							const FString Key = Line.Left(EqualsIdx);
-							if ( Key == ConfigValue.ConfigKey )
-							{
-								// Found the target key, add it to the output and skip the current line if the target value is supposed to replace
-								FileOutput += ConfigValue.ConfigKey + TEXT("=") + ConfigValue.ConfigValue + LINE_TERMINATOR;
-								bShouldExcludeLineFromOutput = ConfigValue.bShouldReplaceExistingValue;
-								bFoundTargetKey = true;
-							}
-						}
-					}
-				}
-
-				// Unless we replaced the key, add this line to the output
-				if ( !bShouldExcludeLineFromOutput )
-				{
-					FileOutput += Line;
-					if ( &LineIn != &FileLines.Last() )
-					{
-						// Add a line terminator on every line except the last
-						FileOutput += LINE_TERMINATOR;
-					}
-				}
-			}
-
-			// If the key did not exist, add it here
-			if ( !bFoundTargetKey )
-			{
-				// If we did not end in the correct section, add the section to the bottom of the file
-				if ( CurSection != TargetSection )
-				{
-					FileOutput += LINE_TERMINATOR;
-					FileOutput += LINE_TERMINATOR;
-					FileOutput += FString::Printf(TEXT("[%s]"), *TargetSection) + LINE_TERMINATOR;
-				}
-
-				// Add the key/value here
-				FileOutput += ConfigValue.ConfigKey + TEXT("=") + ConfigValue.ConfigValue + LINE_TERMINATOR;
-			}
-
-			if ( FFileHelper::SaveStringToFile(FileOutput, *IniFilename) )
-			{
-				bSuccessfullyProcessed = true;
-			}
-		}
-
-		if ( !bSuccessfullyProcessed )
-		{
-			OutFailReason = LOCTEXT("FailedToFixUpDefaultEngine", "Failed to process file DefaultEngine.ini");
-			return false;
-		}
+		return false;
 	}
 
 	// Insert any required feature packs (EG starter content) into ini file. These will be imported automatically when the editor is first run
@@ -1990,46 +1979,45 @@ void GameProjectUtils::DeleteCreatedFiles(const FString& RootFolder, const TArra
 	}
 }
 
-FString GameProjectUtils::GetHardwareConfigString(const FProjectInformation& InProjectInfo)
+void GameProjectUtils::AddHardwareConfigValues(const FProjectInformation& InProjectInfo, TArray<FTemplateConfigValue>& ConfigValues)
 {
-	FString TargetHardwareString;
 	if (InProjectInfo.TargetedHardware.IsSet())
 	{
 		UEnum* HardwareClassEnum = StaticEnum<EHardwareClass::Type>();
 		if (HardwareClassEnum != nullptr)
 		{
+			FString TargetHardwareString;
 			HardwareClassEnum->GetValueAsString(InProjectInfo.TargetedHardware.GetValue(), /*out*/ TargetHardwareString);
+
+			if (!TargetHardwareString.IsEmpty())
+			{
+				ConfigValues.Emplace(TEXT("DefaultEngine.ini"),
+					TEXT("/Script/HardwareTargeting.HardwareTargetingSettings"),
+					TEXT("TargetedHardwareClass"),
+					TargetHardwareString,
+					true /* ShouldReplaceExistingValue */);
+			}
 		}
 	}
 
-	FString GraphicsPresetString;
 	if (InProjectInfo.DefaultGraphicsPerformance.IsSet())
 	{
 		UEnum* GraphicsPresetEnum = StaticEnum<EGraphicsPreset::Type>();
 		if (GraphicsPresetEnum != nullptr)
 		{
+			FString GraphicsPresetString;
 			GraphicsPresetEnum->GetValueAsString(InProjectInfo.DefaultGraphicsPerformance.GetValue(), /*out*/ GraphicsPresetString);
-		}
-	}
-
-	FString HardwareTargeting;
-	if (!TargetHardwareString.IsEmpty() || !GraphicsPresetString.IsEmpty())
-	{
-		HardwareTargeting += LINE_TERMINATOR;
-		HardwareTargeting += TEXT("[/Script/HardwareTargeting.HardwareTargetingSettings]") LINE_TERMINATOR;
-	}
-
-	if (!TargetHardwareString.IsEmpty())
-	{
-		HardwareTargeting += FString::Printf(TEXT("TargetedHardwareClass=%s") LINE_TERMINATOR, *TargetHardwareString);
-	}
 
-	if (!GraphicsPresetString.IsEmpty())
-	{
-		HardwareTargeting += FString::Printf(TEXT("DefaultGraphicsPerformance=%s") LINE_TERMINATOR, *GraphicsPresetString);
+			if (!GraphicsPresetString.IsEmpty())
+			{
+				ConfigValues.Emplace(TEXT("DefaultEngine.ini"),
+					TEXT("/Script/HardwareTargeting.HardwareTargetingSettings"),
+					TEXT("DefaultGraphicsPerformance"),
+					GraphicsPresetString,
+					true /* ShouldReplaceExistingValue */);
+			}
+		}
 	}
-
-	return MoveTemp(HardwareTargeting);
 }
 
 bool GameProjectUtils::GenerateConfigFiles(const FProjectInformation& InProjectInfo, TArray<FString>& OutCreatedFiles, FText& OutFailReason)
@@ -2044,10 +2032,6 @@ bool GameProjectUtils::GenerateConfigFiles(const FProjectInformation& InProjectI
 		const FString DefaultEngineIniFilename = ProjectConfigPath / TEXT("DefaultEngine.ini");
 		FString FileContents;
 
-		FileContents += GetHardwareConfigString(InProjectInfo);
-
-		FileContents += GetRaytracingConfigString(InProjectInfo);
-
 		if(bUseAudioMixerForAllPlatforms)
 		{
 			FileContents += LINE_TERMINATOR;
@@ -2095,6 +2079,15 @@ bool GameProjectUtils::GenerateConfigFiles(const FProjectInformation& InProjectI
 		{
 			return false;
 		}
+
+		TArray<FTemplateConfigValue> ConfigValuesToSet;
+		AddHardwareConfigValues(InProjectInfo, ConfigValuesToSet);
+		AddRaytracingConfigValues(InProjectInfo, ConfigValuesToSet);
+
+		if (!SaveConfigValues(InProjectInfo, ConfigValuesToSet, OutFailReason))
+		{
+			return false;
+		}
 	}
 
 	// DefaultEditor.ini
diff --git a/Engine/Source/Editor/GameProjectGeneration/Public/GameProjectUtils.h b/Engine/Source/Editor/GameProjectGeneration/Public/GameProjectUtils.h
--- a/Engine/Source/Editor/GameProjectGeneration/Public/GameProjectUtils.h
+++ b/Engine/Source/Editor/GameProjectGeneration/Public/GameProjectUtils.h
@@ -12,6 +12,7 @@ class UTemplateProjectDefs;
 class UTemplateCategories;
 struct FProjectDescriptor;
 enum class EClassDomain : uint8;
+struct FTemplateConfigValue;
 
 struct FProjectInformation
 {
@@ -281,7 +282,7 @@ public:
 
 private:
 
-	static FString GetHardwareConfigString(const FProjectInformation& InProjectInfo);
+	static void AddHardwareConfigValues(const FProjectInformation& InProjectInfo, TArray<FTemplateConfigValue>& ConfigValues);
 
 	/** Generates a new project without using a template project */
 	static bool GenerateProjectFromScratch(const FProjectInformation& InProjectInfo, FText& OutFailReason, FText& OutFailLog);
