Restored a line terminator that proved to be necessary for ClassNameRedirects in templates to be correctly spaced out.

This became apparent when creating a project from a template with multiple class redirects.

diff --git a/Engine/Source/Editor/GameProjectGeneration/Private/GameProjectUtils.cpp b/Engine/Source/Editor/GameProjectGeneration/Private/GameProjectUtils.cpp
--- a/Engine/Source/Editor/GameProjectGeneration/Private/GameProjectUtils.cpp
+++ b/Engine/Source/Editor/GameProjectGeneration/Private/GameProjectUtils.cpp
@@ -1520,7 +1520,7 @@ static bool SaveConfigValues(const FProjectInformation& InProjectInfo, const TAr
 							if (Key == ConfigValue.ConfigKey)
 							{
 								// Found the target key, add it to the output and skip the current line if the target value is supposed to replace
-								FileOutput += ConfigValue.ConfigKey + TEXT("=") + ConfigValue.ConfigValue;
+								FileOutput += ConfigValue.ConfigKey + TEXT("=") + ConfigValue.ConfigValue + LINE_TERMINATOR;
 								bShouldExcludeLineFromOutput = ConfigValue.bShouldReplaceExistingValue;
 								bFoundTargetKey = true;
 							}
