Fixed compile error due to missing parameters that I added in 8055321.

diff --git a/Engine/Source/Editor/GameProjectGeneration/Private/SNewProjectWizard.cpp b/Engine/Source/Editor/GameProjectGeneration/Private/SNewProjectWizard.cpp
--- a/Engine/Source/Editor/GameProjectGeneration/Private/SNewProjectWizard.cpp
+++ b/Engine/Source/Editor/GameProjectGeneration/Private/SNewProjectWizard.cpp
@@ -1646,7 +1646,8 @@ TSharedRef<SHorizontalBox> SNewProjectWizard::MakeProjectSettingsOptionsBox()
 		[
 			HardwareTargeting.MakeHardwareClassTargetCombo(
 				FOnHardwareClassChanged::CreateSP(this, &SNewProjectWizard::SetHardwareClassTarget),
-				TAttribute<EHardwareClass::Type>(this, &SNewProjectWizard::GetHardwareClassTarget)
+				TAttribute<EHardwareClass::Type>(this, &SNewProjectWizard::GetHardwareClassTarget),
+				Orient_Horizontal
 			)
 		]
 
@@ -1656,7 +1657,8 @@ TSharedRef<SHorizontalBox> SNewProjectWizard::MakeProjectSettingsOptionsBox()
 		[
 			HardwareTargeting.MakeGraphicsPresetTargetCombo(
 				FOnGraphicsPresetChanged::CreateSP(this, &SNewProjectWizard::SetGraphicsPreset),
-				TAttribute<EGraphicsPreset::Type>(this, &SNewProjectWizard::GetGraphicsPreset)
+				TAttribute<EGraphicsPreset::Type>(this, &SNewProjectWizard::GetGraphicsPreset),
+				Orient_Horizontal
 			)
 		]
 
