Added SRecentProjectBrowser that only displays NumProjects most recently-opened projects.

SProjectBrowser was refactored (lightly) to reuse some of its methods to create only a single category.

diff --git a/Engine/Source/Editor/GameProjectGeneration/Private/SProjectBrowser.cpp b/Engine/Source/Editor/GameProjectGeneration/Private/SProjectBrowser.cpp
--- a/Engine/Source/Editor/GameProjectGeneration/Private/SProjectBrowser.cpp
+++ b/Engine/Source/Editor/GameProjectGeneration/Private/SProjectBrowser.cpp
@@ -61,6 +61,7 @@ struct FProjectItem
 {
 	FText Name;
 	FText Description;
+	FText Category;
 	FString EngineIdentifier;
 	bool bUpToDate;
 	FString ProjectFile;
@@ -68,6 +69,7 @@ struct FProjectItem
 	bool bIsNewProjectItem;
 	TArray<FName> TargetPlatforms;
 	bool bSupportsAllPlatforms;
+	bool bIsSample;
 
 	FProjectItem(const FText& InName, const FText& InDescription, const FString& InEngineIdentifier, bool InUpToDate, const TSharedPtr<FSlateBrush>& InProjectThumbnail, const FString& InProjectFile, bool InIsNewProjectItem, TArray<FName> InTargetPlatforms, bool InSupportsAllPlatforms)
 		: Name(InName)
@@ -130,6 +132,9 @@ void ProjectItemToString(const TSharedPtr<FProjectItem> InItem, TArray<FString>&
 SProjectBrowser::SProjectBrowser()
 	: ProjectItemFilter( ProjectItemTextFilter::FItemToStringArray::CreateStatic( ProjectItemToString ))
 	, NumFilteredProjects(0)
+	, ThumbnailBorderPadding(8)
+	, ThumbnailSize(128)
+	, bPreventSelectionChangeEvent(false)
 {
 
 }
@@ -138,10 +143,6 @@ SProjectBrowser::SProjectBrowser()
 BEGIN_SLATE_FUNCTION_BUILD_OPTIMIZATION
 void SProjectBrowser::Construct( const FArguments& InArgs )
 {
-	bPreventSelectionChangeEvent = false;
-	ThumbnailBorderPadding = 4;
-	ThumbnailSize = 128.0f;
-
 	// Prepare the categories box
 	CategoriesBox = SNew(SVerticalBox);
 
@@ -291,7 +292,7 @@ void SProjectBrowser::Construct( const FArguments& InArgs )
 			]
 
 			+ SVerticalBox::Slot()
-			.Padding( 0, 40, 0, 0 )	// Lots of vertical padding before the dialog buttons at the bottom
+			.Padding( 0, 20, 0, 0 )	// Lots of vertical padding before the dialog buttons at the bottom
 			.AutoHeight()
 			[
 				SNew( SHorizontalBox )
@@ -361,10 +362,12 @@ void SProjectBrowser::Construct( const FArguments& InArgs )
 			break;
 		}
 	}
+
+	ProjectSelectionChangedDelegate = InArgs._OnSelectionChanged;
 }
 END_SLATE_FUNCTION_BUILD_OPTIMIZATION
 
-void SProjectBrowser::ConstructCategory( const TSharedRef<SVerticalBox>& InCategoriesBox, const TSharedRef<FProjectCategory>& Category ) const
+void SProjectBrowser::ConstructCategory( const TSharedRef<SVerticalBox>& InCategoriesBox, const TSharedRef<FProjectCategory>& Category )
 {
 	// Title
 	InCategoriesBox->AddSlot()
@@ -388,18 +391,17 @@ void SProjectBrowser::ConstructCategory( const TSharedRef<SVerticalBox>& InCateg
 	// Project tile view
 	InCategoriesBox->AddSlot()
 	.AutoHeight()
-	.Padding(0.f, 0.0f, 0.f, 40.0f)
 	[
 		SAssignNew(Category->ProjectTileView, STileView<TSharedPtr<FProjectItem> >)
 		.Visibility(this, &SProjectBrowser::GetProjectCategoryVisibility, Category)
 		.ListItemsSource(&Category->FilteredProjectItemsSource)
 		.SelectionMode(ESelectionMode::Single)
-		.ClearSelectionOnClick(false)
+		.ClearSelectionOnClick(true)
 		.AllowOverscroll(EAllowOverscroll::No)
-		.OnGenerateTile(const_cast<SProjectBrowser*>(this), &SProjectBrowser::MakeProjectViewWidget)
+		.OnGenerateTile(this, &SProjectBrowser::MakeProjectViewWidget)
 		.OnContextMenuOpening(this, &SProjectBrowser::OnGetContextMenuContent)
-		.OnMouseButtonDoubleClick(const_cast<SProjectBrowser*>(this), &SProjectBrowser::HandleProjectItemDoubleClick)
-		.OnSelectionChanged(const_cast<SProjectBrowser*>(this), &SProjectBrowser::HandleProjectViewSelectionChanged, Category->CategoryName)
+		.OnMouseButtonDoubleClick(this, &SProjectBrowser::HandleProjectItemDoubleClick)
+		.OnSelectionChanged(this, &SProjectBrowser::HandleProjectViewSelectionChanged, Category->CategoryName)
 		.ItemHeight(ThumbnailSize + ThumbnailBorderPadding + 32)
 		.ItemWidth(ThumbnailSize + ThumbnailBorderPadding)
 	];
@@ -410,6 +412,16 @@ bool SProjectBrowser::HasProjects() const
 	return bHasProjectFiles;
 }
 
+FString SProjectBrowser::GetSelectedProjectFile() const
+{	
+	TSharedPtr<FProjectItem> SelectedItem = GetSelectedProjectItem();
+	if (SelectedItem.IsValid())
+	{
+		return SelectedItem->ProjectFile;
+	}
+
+	return FString();
+}
 
 /* SProjectBrowser implementation
  *****************************************************************************/
@@ -742,6 +754,80 @@ FText SProjectBrowser::GetSelectedProjectName() const
 	return FText::GetEmpty();
 }
 
+static TSharedPtr<FSlateDynamicImageBrush> GetThumbnailForProject(const FString& ProjectFilename)
+{
+	TSharedPtr<FSlateDynamicImageBrush> DynamicBrush;
+	const FString ThumbnailPNGFile = FPaths::GetBaseFilename(ProjectFilename, false) + TEXT(".png");
+	const FString AutoScreenShotPNGFile = FPaths::Combine(*FPaths::GetPath(ProjectFilename), TEXT("Saved"), TEXT("AutoScreenshot.png"));
+	FString PNGFileToUse;
+	if (FPaths::FileExists(*ThumbnailPNGFile))
+	{
+		PNGFileToUse = ThumbnailPNGFile;
+	}
+	else if (FPaths::FileExists(*AutoScreenShotPNGFile))
+	{
+		PNGFileToUse = AutoScreenShotPNGFile;
+	}
+
+	if (!PNGFileToUse.IsEmpty())
+	{
+		const FName BrushName = FName(*PNGFileToUse);
+		DynamicBrush = MakeShareable(new FSlateDynamicImageBrush(BrushName, FVector2D(128, 128)));
+	}
+
+	return DynamicBrush;
+}
+
+static TSharedPtr<FProjectItem> CreateProjectItem(const FString& ProjectFilename)
+{
+	if (FPaths::FileExists(ProjectFilename))
+	{
+		FProjectStatus ProjectStatus;
+		if (IProjectManager::Get().QueryStatusForProject(ProjectFilename, ProjectStatus))
+		{
+			// @todo localized project name
+			const FText ProjectName = FText::FromString(ProjectStatus.Name);
+			const FText ProjectDescription = FText::FromString(ProjectStatus.Description);
+
+			TSharedPtr<FSlateDynamicImageBrush> DynamicBrush = GetThumbnailForProject(ProjectFilename);
+
+			const FString EngineIdentifier = FDesktopPlatformModule::Get()->GetCurrentEngineIdentifier();
+
+			FString ProjectEngineIdentifier;
+			bool bIsUpToDate = FDesktopPlatformModule::Get()->GetEngineIdentifierForProject(ProjectFilename, ProjectEngineIdentifier) 
+				&& ProjectEngineIdentifier == EngineIdentifier;
+
+			// Work out which platforms this project is targeting
+			TArray<FName> TargetPlatforms;
+			for (const PlatformInfo::FPlatformInfo& PlatformInfo : PlatformInfo::GetPlatformInfoArray())
+			{
+				if (PlatformInfo.IsVanilla() && PlatformInfo.PlatformType == PlatformInfo::EPlatformType::Game && ProjectStatus.IsTargetPlatformSupported(PlatformInfo.PlatformInfoName))
+				{
+					TargetPlatforms.Add(PlatformInfo.PlatformInfoName);
+				}
+			}
+			TargetPlatforms.Sort(FNameLexicalLess());
+
+			const bool bIsNewProjectItem = false;
+			TSharedPtr<FProjectItem> ProjectItem = MakeShareable(new FProjectItem(FText::FromString(ProjectStatus.Name),
+				FText::FromString(ProjectStatus.Description),
+				ProjectEngineIdentifier, bIsUpToDate, DynamicBrush, ProjectFilename, bIsNewProjectItem, 
+				TargetPlatforms, ProjectStatus.SupportsAllPlatforms()));
+
+			const FText SamplesCategoryName = LOCTEXT("SamplesCategoryName", "Samples");
+			if (ProjectStatus.bSignedSampleProject)
+			{
+				// Signed samples can't override their category name
+				ProjectItem->Category = SamplesCategoryName;
+			}
+
+			return ProjectItem;
+		}
+	}
+
+	return nullptr;
+}
+
 FReply SProjectBrowser::FindProjects()
 {
 	enum class EProjectCategoryType : uint8
@@ -807,77 +893,26 @@ FReply SProjectBrowser::FindProjects()
 	const FText SamplesCategoryName = LOCTEXT("SamplesCategoryName", "Samples");
 
 	// Add all the discovered projects to the list
-	const FString EngineIdentifier = FDesktopPlatformModule::Get()->GetCurrentEngineIdentifier();
 	for(TMap<FString, EProjectCategoryType>::TConstIterator Iter(AbsoluteProjectFilesToCategory); Iter; ++Iter)
 	{
 		const FString& ProjectFilename = *Iter.Key();
-		const EProjectCategoryType DetectedCategoryType = Iter.Value();
-		if ( FPaths::FileExists(ProjectFilename) )
+		TSharedPtr<FProjectItem> NewProjectItem = CreateProjectItem(ProjectFilename);
+		if (NewProjectItem.IsValid())
 		{
-			FProjectStatus ProjectStatus;
-			if (IProjectManager::Get().QueryStatusForProject(ProjectFilename, ProjectStatus))
+			FText ProjectCategory;
+			
+			if (NewProjectItem->Category.IsEmpty())
 			{
-				// @todo localized project name
-				const FText ProjectName = FText::FromString(ProjectStatus.Name);
-				const FText ProjectDescription = FText::FromString(ProjectStatus.Description);
-
-				TSharedPtr<FSlateDynamicImageBrush> DynamicBrush;
-				const FString ThumbnailPNGFile = FPaths::GetBaseFilename(ProjectFilename, false) + TEXT(".png");
-				const FString AutoScreenShotPNGFile = FPaths::Combine(*FPaths::GetPath(ProjectFilename), TEXT("Saved"), TEXT("AutoScreenshot.png"));
-				FString PNGFileToUse;
-				if ( FPaths::FileExists(*ThumbnailPNGFile) )
-				{
-					PNGFileToUse = ThumbnailPNGFile;
-				}
-				else if ( FPaths::FileExists(*AutoScreenShotPNGFile) )
-				{
-					PNGFileToUse = AutoScreenShotPNGFile;
-				}
-
-				if ( !PNGFileToUse.IsEmpty() )
-				{
-					const FName BrushName = FName(*PNGFileToUse);
-					DynamicBrush = MakeShareable( new FSlateDynamicImageBrush(BrushName, FVector2D(128,128) ) );
-				}
-
-				FText ProjectCategory;
-				if ( ProjectStatus.bSignedSampleProject )
-				{
-					// Signed samples can't override their category name
-					ProjectCategory = SamplesCategoryName;
-				}
-				else
-				{
-					if(ProjectStatus.Category.IsEmpty())
-					{
-						// Not category specified, so use the category for the detected project type
-						ProjectCategory = (DetectedCategoryType == EProjectCategoryType::Sample) ? SamplesCategoryName : MyProjectsCategoryName;
-					}
-					else
-					{
-						// Use the user defined category
-						ProjectCategory = FText::FromString(ProjectStatus.Category);
-					}
-				}
-
-				FString ProjectEngineIdentifier;
-				bool bIsUpToDate = FDesktopPlatformModule::Get()->GetEngineIdentifierForProject(ProjectFilename, ProjectEngineIdentifier) && ProjectEngineIdentifier == EngineIdentifier;
-
-				// Work out which platforms this project is targeting
-				TArray<FName> TargetPlatforms;
-				for(const PlatformInfo::FPlatformInfo& PlatformInfo : PlatformInfo::GetPlatformInfoArray())
-				{
-					if(PlatformInfo.IsVanilla() && PlatformInfo.PlatformType == PlatformInfo::EPlatformType::Game && ProjectStatus.IsTargetPlatformSupported(PlatformInfo.PlatformInfoName))
-					{
-						TargetPlatforms.Add(PlatformInfo.PlatformInfoName);
-					}
-				}
-				TargetPlatforms.Sort(FNameLexicalLess());
-
-				const bool bIsNewProjectItem = false;
-				TSharedRef<FProjectItem> NewProjectItem = MakeShareable( new FProjectItem(ProjectName, ProjectDescription, ProjectEngineIdentifier, bIsUpToDate, DynamicBrush, ProjectFilename, bIsNewProjectItem, TargetPlatforms, ProjectStatus.SupportsAllPlatforms() ) );
-				AddProjectToCategory(NewProjectItem, ProjectCategory);
+				// Not category specified, so use the category for the detected project type
+				ProjectCategory = (Iter.Value() == EProjectCategoryType::Sample) ? SamplesCategoryName : MyProjectsCategoryName;
+			}
+			else
+			{
+				// Use the user defined category
+				ProjectCategory = NewProjectItem->Category;
 			}
+
+			AddProjectToCategory(NewProjectItem.ToSharedRef(), ProjectCategory);
 		}
 	}
 
@@ -1316,6 +1351,19 @@ FReply SProjectBrowser::OnBrowseToProjectClicked()
 	return FReply::Handled();
 }
 
+void SProjectBrowser::ClearSelection()
+{
+	TGuardValue<bool> SelectionEventGuard(bPreventSelectionChangeEvent, true);
+
+	for (const TSharedRef<FProjectCategory>& Category : ProjectCategories)
+	{
+		if (Category->ProjectTileView.IsValid())
+		{
+			Category->ProjectTileView->ClearSelection();
+		}
+	}
+}
+
 
 void SProjectBrowser::HandleProjectViewSelectionChanged(TSharedPtr<FProjectItem> ProjectItem, ESelectInfo::Type SelectInfo, FText CategoryName)
 {
@@ -1323,21 +1371,27 @@ void SProjectBrowser::HandleProjectViewSelectionChanged(TSharedPtr<FProjectItem>
 	{
 		TGuardValue<bool> SelectionEventGuard(bPreventSelectionChangeEvent, true);
 
-		for ( auto CategoryIt = ProjectCategories.CreateConstIterator(); CategoryIt; ++CategoryIt )
+		for (const TSharedRef<FProjectCategory>& Category : ProjectCategories)
 		{
-			const TSharedRef<FProjectCategory>& Category = *CategoryIt;
-			if ( Category->ProjectTileView.IsValid() && !Category->CategoryName.EqualToCaseIgnored(CategoryName) )
+			if (Category->ProjectTileView.IsValid() && !Category->CategoryName.EqualToCaseIgnored(CategoryName))
 			{
 				Category->ProjectTileView->ClearSelection();
 			}
 		}
+	}
 
-		const TSharedPtr<FProjectItem> SelectedItem = GetSelectedProjectItem();
-		if ( SelectedItem.IsValid() && SelectedItem != CurrentlySelectedItem )
-		{
-			CurrentSelectedProjectPath = FText::FromString( SelectedItem->ProjectFile );
-		}
+	FString ProjectFile;
+	if (ProjectItem.IsValid())
+	{
+		ProjectFile = ProjectItem->ProjectFile;
+		CurrentSelectedProjectPath = FText::FromString(ProjectFile);
 	}
+	else
+	{
+		CurrentSelectedProjectPath = FText::GetEmpty();
+	}
+
+	ProjectSelectionChangedDelegate.ExecuteIfBound(ProjectFile);
 }
 
 
@@ -1432,4 +1486,73 @@ FText SProjectBrowser:: GetItemHighlightText() const
 	return ProjectItemFilter.GetRawFilterText();
 }
 
+
+SRecentProjectBrowser::SRecentProjectBrowser()
+{
+	
+}
+
+void SRecentProjectBrowser::Construct(const FArguments& InArgs)
+{
+	CategoriesBox = SNew(SVerticalBox);
+
+	const FString EngineIdentifier = FDesktopPlatformModule::Get()->GetCurrentEngineIdentifier();
+
+	TArray<FString> RecentProjects = GetDefault<UEditorSettings>()->RecentlyOpenedProjectFiles;
+	for (int32 Idx = 0; Idx < InArgs._NumProjects && Idx < RecentProjects.Num(); ++Idx)
+	{
+		const FString ProjectFilename = RecentProjects[Idx];
+		
+		TSharedPtr<FProjectItem> ProjectItem = CreateProjectItem(ProjectFilename);
+
+		if (ProjectItem)
+		{
+			AddProjectToCategory(ProjectItem.ToSharedRef(), LOCTEXT("ProjectDialog_Recent", "Recent Projects"));
+		}
+	}
+
+	check(ProjectCategories.Num() == 1);
+
+	PopulateFilteredProjectCategories();
+
+	ConstructCategory(CategoriesBox.ToSharedRef(), ProjectCategories[0]);
+
+	ChildSlot
+	[
+		SNew(SBorder)
+		.BorderImage(FEditorStyle::GetBrush("ToolPanel.GroupBorder"))
+		.Padding(FMargin(8.f, 4.f))
+		[
+			SNew(SScrollBox)
+			+ SScrollBox::Slot()
+			[
+				CategoriesBox.ToSharedRef()
+			]
+		]
+	];
+
+	if (ProjectCategories.Num() > 0)
+	{
+		const TSharedRef<FProjectCategory> Category = ProjectCategories[0];
+		if (ensure(Category->ProjectItemsSource.Num() > 0) && ensure(Category->ProjectTileView.IsValid()))
+		{
+			Category->ProjectTileView->SetSelection(Category->ProjectItemsSource[0], ESelectInfo::Direct);
+		}
+	}
+
+	bHasProjectFiles = false;
+	for (auto CategoryIt = ProjectCategories.CreateConstIterator(); CategoryIt; ++CategoryIt)
+	{
+		const TSharedRef<FProjectCategory>& Category = *CategoryIt;
+
+		if (Category->ProjectItemsSource.Num() > 0)
+		{
+			bHasProjectFiles = true;
+			break;
+		}
+	}
+
+	ProjectSelectionChangedDelegate = InArgs._OnSelectionChanged;
+}
+
 #undef LOCTEXT_NAMESPACE
diff --git a/Engine/Source/Editor/GameProjectGeneration/Private/SProjectBrowser.h b/Engine/Source/Editor/GameProjectGeneration/Private/SProjectBrowser.h
--- a/Engine/Source/Editor/GameProjectGeneration/Private/SProjectBrowser.h
+++ b/Engine/Source/Editor/GameProjectGeneration/Private/SProjectBrowser.h
@@ -19,6 +19,8 @@ struct FProjectItem;
 struct FSlateBrush;
 enum class ECheckBoxState : uint8;
 
+DECLARE_DELEGATE_OneParam(FProjectSelectionChanged, FString);
+
 /**
  * A list of known projects with the option to add a new one
  */
@@ -31,6 +33,8 @@ public:
 
 	SLATE_BEGIN_ARGS(SProjectBrowser) { }
 
+	SLATE_ARGUMENT(FProjectSelectionChanged, OnSelectionChanged);
+
 	SLATE_END_ARGS()
 
 public:
@@ -47,9 +51,16 @@ public:
 
 	bool HasProjects() const;
 
+	void ClearSelection();
+
+	FString GetSelectedProjectFile() const;
+
+	/** Begins the opening process for the selected project */
+	void OpenSelectedProject();
+
 protected:
 
-	void ConstructCategory( const TSharedRef<SVerticalBox>& CategoriesBox, const TSharedRef<FProjectCategory>& Category ) const;
+	void ConstructCategory( const TSharedRef<SVerticalBox>& CategoriesBox, const TSharedRef<FProjectCategory>& Category );
 
 	/** Creates a row in the template list */
 	TSharedRef<ITableRow> MakeProjectViewWidget( TSharedPtr<FProjectItem> ProjectItem, const TSharedRef<STableViewBase>& OwnerTable );
@@ -87,12 +98,9 @@ protected:
 	/** Opens the specified project file */
 	bool OpenProject( const FString& ProjectFile );
 
-	/** Begins the opening process for the selected project */
-	void OpenSelectedProject( );
-
 	/** Populate the list of filtered project categories */
 	void PopulateFilteredProjectCategories();
-	
+
 	/**
 	 * Called after a key is pressed when this widget has focus (this event bubbles if not handled)
 	 *
@@ -103,6 +111,17 @@ protected:
 	 */
 	virtual FReply OnKeyDown( const FGeometry& MyGeometry, const FKeyEvent& InKeyEvent ) override;
 
+protected:
+
+	/** Holds the collection of project categories. */
+	TArray<TSharedRef<FProjectCategory> > ProjectCategories;
+
+	bool bHasProjectFiles;
+
+	TSharedPtr<SVerticalBox> CategoriesBox;
+
+	FProjectSelectionChanged ProjectSelectionChangedDelegate;
+
 private:
 
 	FReply OnBrowseToProjectClicked();
@@ -143,8 +162,6 @@ private:
 	FText GetItemHighlightText() const;
 
 private:
-	// Holds the collection of project categories.
-	TArray<TSharedRef<FProjectCategory> > ProjectCategories;
 
 	/** Search box used to set the filter text */
 	TSharedPtr<class SSearchBox> SearchBoxPtr;
@@ -162,11 +179,25 @@ private:
 	FText CurrentSelectedProjectPath;
 
 	bool IsOnlineContentFinished;
-	TSharedPtr<SVerticalBox> CategoriesBox;
-
-	bool bHasProjectFiles;
-private:
 
 	// Holds a delegate that is executed when the new project screen is being requested.
 	FNewProjectScreenRequested NewProjectScreenRequestedDelegate;
 };
+
+/** Class to only display a finite number of most-recent projects. */
+class SRecentProjectBrowser : public SProjectBrowser
+{
+	SLATE_BEGIN_ARGS(SRecentProjectBrowser) 
+	{
+		_NumProjects = 10;
+	}
+
+	SLATE_ARGUMENT(int32, NumProjects);
+	SLATE_EVENT(FProjectSelectionChanged, OnSelectionChanged);
+
+	SLATE_END_ARGS()
+
+	SRecentProjectBrowser();
+
+	void Construct(const FArguments& InArgs);
+};
\ No newline at end of file
