Project templates have been split into more explicit categories, where templates can be enabled for either blueprint or code generation, thus removing the need for C++ and Blueprint categories in favour of categories like Virtual Production or Visualization.

Template categories can be configured from TemplateCategories.ini in the Templates/ folder. Templates can now also exist in multiple categories.

diff --git a/Engine/Source/Editor/GameProjectGeneration/Classes/TemplateProjectDefs.h b/Engine/Source/Editor/GameProjectGeneration/Classes/TemplateProjectDefs.h
--- a/Engine/Source/Editor/GameProjectGeneration/Classes/TemplateProjectDefs.h
+++ b/Engine/Source/Editor/GameProjectGeneration/Classes/TemplateProjectDefs.h
@@ -6,6 +6,7 @@
 #include "UObject/ObjectMacros.h"
 #include "UObject/Object.h"
 #include "FeaturePackContentSource.h"
+#include "Internationalization/PolyglotTextData.h"
 #include "TemplateProjectDefs.generated.h"
 
 // does not require reflection exposure
@@ -23,7 +24,7 @@ struct FTemplateConfigValue
 USTRUCT()
 struct FTemplateReplacement
 {
-	GENERATED_USTRUCT_BODY()
+	GENERATED_BODY()
 
 	UPROPERTY()
 	TArray<FString> Extensions;
@@ -41,7 +42,7 @@ struct FTemplateReplacement
 USTRUCT()
 struct FTemplateFolderRename
 {
-	GENERATED_USTRUCT_BODY()
+	GENERATED_BODY()
 
 	UPROPERTY()
 	FString From;
@@ -53,18 +54,32 @@ struct FTemplateFolderRename
 USTRUCT()
 struct FLocalizedTemplateString
 {
-	GENERATED_USTRUCT_BODY()
+	GENERATED_BODY()
 
 	UPROPERTY()
 	FString Language;
 
 	UPROPERTY()
 	FString Text;
+
+	/** Find a text value for the current locale (or fallback to English) from the given array of localized strings. */
+	static FText GetLocalizedText(const TArray<FLocalizedTemplateString>& LocalizedStrings);
+};
+
+UENUM()
+enum class ETemplateSetting
+{
+	Languages,
+	HardwareTarget,
+	GraphicsPreset,
+	StarterContent,
+	VR
 };
 
 UCLASS(abstract,config=TemplateDefs,MinimalAPI)
 class UTemplateProjectDefs : public UObject
 {
+public:
 	GENERATED_UCLASS_BODY()
 
 	UPROPERTY(config)
@@ -92,7 +107,7 @@ class UTemplateProjectDefs : public UObject
 	FString SortKey;
 
 	UPROPERTY(config)
-	FName Category;
+	TArray<FName> Categories;
 
 	UPROPERTY(config)
 	FString ClassTypes;
@@ -103,6 +118,14 @@ class UTemplateProjectDefs : public UObject
 	/* Should we allow creation of a project from this template */
 	UPROPERTY(config)
 	bool bAllowProjectCreation;
+
+	/* Should we skip the project settings page when creating a project from this template */
+	UPROPERTY(config)
+	bool bSkipProjectSettings;
+
+	/* Optional list of settings to hide. If none are specified, then all settings are shown. */
+	UPROPERTY(config)
+	TArray<ETemplateSetting> HiddenSettings;
 	
 	/* Optional list of feature packs to include */
 	UPROPERTY(config)
@@ -144,3 +167,40 @@ private:
 	void FixString(FString& InOutStringToFix, const FString& TemplateName, const FString& ProjectName);
 };
 
+USTRUCT()
+struct FTemplateCategoryDef
+{
+	GENERATED_BODY()
+	
+	/** Key to use for identifying what category a template is in. */
+	UPROPERTY()
+	FName Key;
+
+	/** Localized name for this template category. */
+	UPROPERTY()
+	TArray<FLocalizedTemplateString> LocalizedDisplayNames;
+
+	/** Localized description for this template category. */
+	UPROPERTY()
+	TArray<FLocalizedTemplateString> LocalizedDescriptions;
+
+	UPROPERTY()
+	FName Icon;
+
+	UPROPERTY()
+	FName Image;
+
+	UPROPERTY()
+	bool IsMajorCategory;
+};
+
+UCLASS(config=TemplateCategories)
+class UTemplateCategories : public UObject
+{
+public:
+	GENERATED_BODY()
+
+	/** Array of all categories specified in this location. */
+	UPROPERTY(config)
+	TArray<FTemplateCategoryDef> Categories;
+};
\ No newline at end of file
diff --git a/Engine/Source/Editor/GameProjectGeneration/Private/GameProjectGenerationModule.cpp b/Engine/Source/Editor/GameProjectGeneration/Private/GameProjectGenerationModule.cpp
--- a/Engine/Source/Editor/GameProjectGeneration/Private/GameProjectGenerationModule.cpp
+++ b/Engine/Source/Editor/GameProjectGeneration/Private/GameProjectGenerationModule.cpp
@@ -10,57 +10,42 @@
 #include "SGameProjectDialog.h"
 #include "SNewClassDialog.h"
 #include "TemplateCategory.h"
-#include "SourceCodeNavigation.h"
-
+#include "HAL/FileManager.h"
+#include "Interfaces/IPluginManager.h"
+#include "TemplateProjectDefs.h"
+#include "Internationalization/Culture.h"
 
 IMPLEMENT_MODULE( FGameProjectGenerationModule, GameProjectGeneration );
 DEFINE_LOG_CATEGORY(LogGameProjectGeneration);
 
 #define LOCTEXT_NAMESPACE "GameProjectGeneration"
 
-FName FTemplateCategory::BlueprintCategoryName = "Blueprint";
-FName FTemplateCategory::CodeCategoryName = "C++";
-FName FTemplateCategory::EnterpriseCategoryName = "Unreal Studio";
-
 void FGameProjectGenerationModule::StartupModule()
 {
-	RegisterTemplateCategory(
-		FTemplateCategory::BlueprintCategoryName,
-		LOCTEXT("BlueprintCategory_Name", "Blueprint"),
-		LOCTEXT("BlueprintCategory_Description", "Blueprint templates require no programming knowledge.\nAll game mechanics can be implemented using Blueprint visual scripting.\nEach template includes a basic set of blueprints to use as a starting point for your game."),
-		FEditorStyle::GetBrush("GameProjectDialog.BlueprintIcon"),
-		FEditorStyle::GetBrush("GameProjectDialog.BlueprintImage"));
-
-	RegisterTemplateCategory(
-		FTemplateCategory::CodeCategoryName,
-		LOCTEXT("CodeCategory_Name", "C++"),
-		FText::Format(
-			LOCTEXT("CodeCategory_Description", "C++ templates offer a good example of how to work with some of the core concepts of the Engine from code.\nYou still have the option of adding your own blueprints to the project at a later date if you want.\nChoosing this template type requires you to have {0} installed."),
-			FSourceCodeNavigation::GetSuggestedSourceCodeIDE()
-		),
-		FEditorStyle::GetBrush("GameProjectDialog.CodeIcon"),
-		FEditorStyle::GetBrush("GameProjectDialog.CodeImage"));
-
-	RegisterTemplateCategory(
-		FTemplateCategory::EnterpriseCategoryName,
-		LOCTEXT("EnterpriseCategory_Name", "Unreal Studio"),
-		LOCTEXT("EnterpriseCategory_Description", "Unreal Studio blueprint templates require no programming knowledge.\nEach template include a basic set of blueprints to use as a starting point for your Unreal Studio project."),
-		FEditorStyle::GetBrush("GameProjectDialog.BlueprintIcon"),
-		FEditorStyle::GetBrush("GameProjectDialog.BlueprintImage"));
+	LoadTemplateCategories();
 }
 
-
 void FGameProjectGenerationModule::ShutdownModule()
 {
 	
 }
 
-
 TSharedRef<class SWidget> FGameProjectGenerationModule::CreateGameProjectDialog(bool bAllowProjectOpening, bool bAllowProjectCreate)
 {
-	return SNew(SGameProjectDialog)
-		.AllowProjectOpening(bAllowProjectOpening)
-		.AllowProjectCreate(bAllowProjectCreate);
+	ensure(bAllowProjectOpening || bAllowProjectCreate);
+
+	SGameProjectDialog::EMode Mode = SGameProjectDialog::EMode::Both;
+	
+	if (bAllowProjectOpening && !bAllowProjectCreate)
+	{
+		Mode = SGameProjectDialog::EMode::Open;
+	}
+	else if (bAllowProjectCreate && !bAllowProjectOpening)
+	{
+		Mode = SGameProjectDialog::EMode::New;
+	}
+
+	return SNew(SGameProjectDialog, Mode);
 }
 
 
@@ -161,21 +146,61 @@ void FGameProjectGenerationModule::ClearSupportedTargetPlatforms()
 	GameProjectUtils::ClearSupportedTargetPlatforms();
 }
 
-bool FGameProjectGenerationModule::RegisterTemplateCategory(FName Type, FText Name, FText Description, const FSlateBrush* Icon, const FSlateBrush* Image)
+void FGameProjectGenerationModule::LoadTemplateCategories()
 {
-	if (TemplateCategories.Contains(Type))
+	// Now discover and all data driven templates
+	TArray<FString> TemplateRootFolders;
+
+	// @todo rocket make template folder locations extensible.
+	TemplateRootFolders.Add(FPaths::RootDir() / TEXT("Templates"));
+
+	// Add the Enterprise templates
+	TemplateRootFolders.Add(FPaths::EnterpriseDir() / TEXT("Templates"));
+
+	// allow plugins to define templates
+	TArray<TSharedRef<IPlugin>> Plugins = IPluginManager::Get().GetEnabledPlugins();
+	for (const TSharedRef<IPlugin>& Plugin : Plugins)
 	{
-		return false;
+		FString PluginDirectory = Plugin->GetBaseDir();
+		if (!PluginDirectory.IsEmpty())
+		{
+			const FString PluginTemplatesDirectory = PluginDirectory / TEXT("Templates");
+
+			if (IFileManager::Get().DirectoryExists(*PluginTemplatesDirectory))
+			{
+				TemplateRootFolders.Add(PluginTemplatesDirectory);
+			}
+		}
 	}
 
-	FTemplateCategory Category = { Name, Description, Icon, Image, Type };
-	TemplateCategories.Add(Type, MakeShareable(new FTemplateCategory(Category)));
-	return true;
-}
-
-void FGameProjectGenerationModule::UnRegisterTemplateCategory(FName Type)
-{
-	TemplateCategories.Remove(Type);
+	for (const FString& Root : TemplateRootFolders)
+	{
+		UTemplateCategories* CategoryDefs = GameProjectUtils::LoadTemplateCategories(Root);
+		if (CategoryDefs != nullptr)
+		{
+			for (const FTemplateCategoryDef& Category : CategoryDefs->Categories)
+			{
+				TSharedPtr<FTemplateCategory> TemplateCategory;
+
+				TSharedPtr<FTemplateCategory>* Existing = TemplateCategories.Find(Category.Key);
+				if (Existing == nullptr)
+				{
+					TemplateCategory = TemplateCategories.Add(Category.Key, MakeShareable(new FTemplateCategory));
+				}
+				else
+				{
+					TemplateCategory = *Existing;
+				}
+
+				TemplateCategory->Key = Category.Key;
+				TemplateCategory->DisplayName = FLocalizedTemplateString::GetLocalizedText(Category.LocalizedDisplayNames);
+				TemplateCategory->Description = FLocalizedTemplateString::GetLocalizedText(Category.LocalizedDescriptions);
+				TemplateCategory->Icon = FEditorStyle::GetBrush(Category.Icon);
+				TemplateCategory->Image = FEditorStyle::GetBrush(Category.Image);
+				TemplateCategory->IsMajor = Category.IsMajorCategory;
+			}
+		}
+	}
 }
 
 #undef LOCTEXT_NAMESPACE
diff --git a/Engine/Source/Editor/GameProjectGeneration/Private/GameProjectUtils.cpp b/Engine/Source/Editor/GameProjectGeneration/Private/GameProjectUtils.cpp
--- a/Engine/Source/Editor/GameProjectGeneration/Private/GameProjectUtils.cpp
+++ b/Engine/Source/Editor/GameProjectGeneration/Private/GameProjectUtils.cpp
@@ -799,6 +799,7 @@ bool GameProjectUtils::CreateProject(const FProjectInformation& InProjectInfo, F
 	{
 		TArray<FAnalyticsEventAttribute> EventAttributes;
 		EventAttributes.Add(FAnalyticsEventAttribute(TEXT("Template"), TemplateName));
+		EventAttributes.Add(FAnalyticsEventAttribute(TEXT("Category"), InProjectInfo.TemplateCategory.ToString()));
 		EventAttributes.Add(FAnalyticsEventAttribute(TEXT("ProjectType"), InProjectInfo.bShouldGenerateCode ? TEXT("C++ Code") : TEXT("Content Only")));
 		EventAttributes.Add(FAnalyticsEventAttribute(TEXT("Outcome"), bProjectCreationSuccessful ? TEXT("Successful") : TEXT("Failed")));
 
@@ -807,8 +808,7 @@ bool GameProjectUtils::CreateProject(const FProjectInformation& InProjectInfo, F
 		Enum = StaticEnum<EGraphicsPreset::Type>();
 		EventAttributes.Add(FAnalyticsEventAttribute(TEXT("GraphicsPreset"), Enum ? Enum->GetNameStringByValue(InProjectInfo.DefaultGraphicsPerformance) : FString()));
 		EventAttributes.Add(FAnalyticsEventAttribute(TEXT("StarterContent"), InProjectInfo.bCopyStarterContent ? TEXT("Yes") : TEXT("No")));
-		EventAttributes.Emplace(TEXT("Enterprise"), InProjectInfo.bIsEnterpriseProject);
-
+		
 		FEngineAnalytics::GetProvider().RecordEvent( TEXT( "Editor.NewProject.ProjectCreated" ), EventAttributes );
 	}
 
@@ -1217,9 +1217,23 @@ GameProjectUtils::EAddCodeToProjectResult GameProjectUtils::AddCodeToProject(con
 	return Result;
 }
 
+UTemplateCategories* GameProjectUtils::LoadTemplateCategories(const FString& RootDir)
+{
+	UTemplateCategories* TemplateCategories = nullptr;
+
+	FString TemplateCategoriesIniFilename = RootDir / TEXT("TemplateCategories.ini");
+	if (FPlatformFileManager::Get().GetPlatformFile().FileExists(*TemplateCategoriesIniFilename))
+	{
+		TemplateCategories = NewObject<UTemplateCategories>();
+		TemplateCategories->LoadConfig(UTemplateCategories::StaticClass(), *TemplateCategoriesIniFilename);
+	}
+
+	return TemplateCategories;
+}
+
 UTemplateProjectDefs* GameProjectUtils::LoadTemplateDefs(const FString& ProjectDirectory)
 {
-	UTemplateProjectDefs* TemplateDefs = NULL;
+	UTemplateProjectDefs* TemplateDefs = nullptr;
 
 	const FString TemplateDefsIniFilename = ProjectDirectory / TEXT("Config") / GetTemplateDefsFilename();
 	if ( FPlatformFileManager::Get().GetPlatformFile().FileExists(*TemplateDefsIniFilename) )
diff --git a/Engine/Source/Editor/GameProjectGeneration/Private/TemplateCategory.h b/Engine/Source/Editor/GameProjectGeneration/Private/TemplateCategory.h
--- a/Engine/Source/Editor/GameProjectGeneration/Private/TemplateCategory.h
+++ b/Engine/Source/Editor/GameProjectGeneration/Private/TemplateCategory.h
@@ -8,10 +8,8 @@ struct FSlateBrush;
 
 struct FTemplateCategory
 {
-	static FName BlueprintCategoryName, CodeCategoryName, EnterpriseCategoryName;
-
 	/** Localised name of this category */
-	FText Name;
+	FText DisplayName;
 	
 	/** A description of the templates contained within this category */
 	FText Description;
@@ -22,6 +20,9 @@ struct FTemplateCategory
 	/** A thumbnail to help identify this category (on the screenshot)*/
 	const FSlateBrush* Image;
 
-	/** A unique name for this category */
-	FName Type;
+	/** A unique key for this category */
+	FName Key;
+
+	/** Is this a major or minor category? Controls where in the UI it will show up. */
+	bool IsMajor;
 };
diff --git a/Engine/Source/Editor/GameProjectGeneration/Private/TemplateItem.h b/Engine/Source/Editor/GameProjectGeneration/Private/TemplateItem.h
--- a/Engine/Source/Editor/GameProjectGeneration/Private/TemplateItem.h
+++ b/Engine/Source/Editor/GameProjectGeneration/Private/TemplateItem.h
@@ -4,40 +4,27 @@
 
 #include "CoreMinimal.h"
 #include "Styling/SlateBrush.h"
+#include "TemplateProjectDefs.h"
 
 /** Struct describing a single template project */
 struct FTemplateItem
 {
 	FText		Name;
 	FText		Description;
-	bool		bGenerateCode;
-	FName		Type;
+	TArray<FName> Categories;
 
+	FString		Key;
 	FString		SortKey;
-	FString		ProjectFile;
 
 	TSharedPtr<FSlateBrush> Thumbnail;
 	TSharedPtr<FSlateBrush> PreviewImage;
 
 	FString		ClassTypes;
 	FString		AssetTypes;
-	FTemplateItem(FText InName, FText InDescription, bool bInGenerateCode, FName InType, FString InSortKey, FString InProjectFile, TSharedPtr<FSlateBrush> InThumbnail, TSharedPtr<FSlateBrush> InPreviewImage,FString InClassTypes, FString InAssetTypes)
-		: Name(InName), Description(InDescription), bGenerateCode(bInGenerateCode), Type(InType), SortKey(MoveTemp(InSortKey)), ProjectFile(MoveTemp(InProjectFile)), Thumbnail(InThumbnail), PreviewImage(InPreviewImage)
-		, ClassTypes(InClassTypes), AssetTypes(InAssetTypes)
-	{}
-
-	FTemplateItem(const FTemplateItem& InItem)
-	{
-		Name = InItem.Name;
-		Description = InItem.Description;
-		bGenerateCode = InItem.bGenerateCode;
-		Type = InItem.Type;
-
-		SortKey = InItem.SortKey;
-		ProjectFile = InItem.ProjectFile;
-
-		ClassTypes = InItem.ClassTypes;
-		AssetTypes = InItem.AssetTypes;
-	}
-};
 
+	FString		CodeProjectFile;
+	FString		BlueprintProjectFile;
+
+	bool		bSkipProjectSettings;
+	TArray<ETemplateSetting> HiddenSettings;
+};
diff --git a/Engine/Source/Editor/GameProjectGeneration/Private/TemplateProjectDefs.cpp b/Engine/Source/Editor/GameProjectGeneration/Private/TemplateProjectDefs.cpp
--- a/Engine/Source/Editor/GameProjectGeneration/Private/TemplateProjectDefs.cpp
+++ b/Engine/Source/Editor/GameProjectGeneration/Private/TemplateProjectDefs.cpp
@@ -4,10 +4,7 @@
 #include "TemplateProjectDefs.h"
 #include "Internationalization/Culture.h"
 
-namespace
-{
-
-FText GetLocalizedText(const TArray<FLocalizedTemplateString>& LocalizedStrings)
+FText FLocalizedTemplateString::GetLocalizedText(const TArray<FLocalizedTemplateString>& LocalizedStrings)
 {
 	auto FindLocalizedStringForCulture = [&LocalizedStrings](const FString& InCulture)
 	{
@@ -49,8 +46,6 @@ FText GetLocalizedText(const TArray<FLocalizedTemplateString>& LocalizedStrings)
 	return FText();
 }
 
-}
-
 FTemplateConfigValue::FTemplateConfigValue(const FString& InFile, const FString& InSection, const FString& InKey, const FString& InValue, bool InShouldReplaceExistingValue)
 	: ConfigFile(InFile)
 	, ConfigSection(InSection)
@@ -61,8 +56,6 @@ FTemplateConfigValue::FTemplateConfigValue(const FString& InFile, const FString&
 }
 
 
-
-
 UTemplateProjectDefs::UTemplateProjectDefs(const FObjectInitializer& ObjectInitializer)
 	: Super(ObjectInitializer)
 {
@@ -106,12 +99,12 @@ void UTemplateProjectDefs::FixupStrings(const FString& TemplateName, const FStri
 
 FText UTemplateProjectDefs::GetDisplayNameText()
 {
-	return GetLocalizedText(LocalizedDisplayNames);
+	return FLocalizedTemplateString::GetLocalizedText(LocalizedDisplayNames);
 }
 
 FText UTemplateProjectDefs::GetLocalizedDescription()
 {
-	return GetLocalizedText(LocalizedDescriptions);
+	return FLocalizedTemplateString::GetLocalizedText(LocalizedDescriptions);
 }
 
 void UTemplateProjectDefs::FixString(FString& InOutStringToFix, const FString& TemplateName, const FString& ProjectName)
diff --git a/Engine/Source/Editor/GameProjectGeneration/Public/GameProjectGenerationModule.h b/Engine/Source/Editor/GameProjectGeneration/Public/GameProjectGenerationModule.h
--- a/Engine/Source/Editor/GameProjectGeneration/Public/GameProjectGenerationModule.h
+++ b/Engine/Source/Editor/GameProjectGeneration/Public/GameProjectGenerationModule.h
@@ -120,16 +120,16 @@ public:
 
 public:
 
-	/** (Un)register a new type of template category to be shown on the new project page */
-	virtual bool RegisterTemplateCategory(FName Type, FText Name, FText Description, const FSlateBrush* Icon, const FSlateBrush* Image);
-	virtual void UnRegisterTemplateCategory(FName Type);
-
 	// Non DLL-exposed access to template categories
 	TSharedPtr<const FTemplateCategory> GetCategory(FName Type) const { return TemplateCategories.FindRef(Type); }
 
+	void GetAllTemplateCategories(TArray<TSharedPtr<FTemplateCategory>>& OutCategories) const { TemplateCategories.GenerateValueArray(OutCategories); }
+
 private:
 	FAddCodeToProjectDialogOpenedEvent AddCodeToProjectDialogOpenedEvent;
 
 	/** Map of template categories from type to ptr */
 	FTemplateCategoryMap TemplateCategories;
+
+	void LoadTemplateCategories();
 };
diff --git a/Engine/Source/Editor/GameProjectGeneration/Public/GameProjectUtils.h b/Engine/Source/Editor/GameProjectGeneration/Public/GameProjectUtils.h
--- a/Engine/Source/Editor/GameProjectGeneration/Public/GameProjectUtils.h
+++ b/Engine/Source/Editor/GameProjectGeneration/Public/GameProjectUtils.h
@@ -9,33 +9,27 @@
 #include "HardwareTargetingSettings.h"
 
 class UTemplateProjectDefs;
+class UTemplateCategories;
 struct FProjectDescriptor;
 enum class EClassDomain : uint8;
 
 struct FProjectInformation
 {
-	FProjectInformation(FString InProjectFilename, bool bInGenerateCode, bool bInCopyStarterContent, FString InTemplateFile = FString())
-		: ProjectFilename(MoveTemp(InProjectFilename))
-		, TemplateFile(MoveTemp(InTemplateFile))
-		, bShouldGenerateCode(bInGenerateCode)
-		, bCopyStarterContent(bInCopyStarterContent)
-		, bIsEnterpriseProject(false)
-		, bForceExtendedLuminanceRange(false)
-		, TargetedHardware(EHardwareClass::Desktop)
-		, DefaultGraphicsPerformance(EGraphicsPreset::Maximum)
-	{
-	}
+	FProjectInformation() = default;
 
 	FString ProjectFilename;
 	FString TemplateFile;
+	FName TemplateCategory;
 
-	bool bShouldGenerateCode;
-	bool bCopyStarterContent;
-	bool bIsEnterpriseProject;
-	bool bForceExtendedLuminanceRange; // See "r.DefaultFeature.AutoExposure.ExtendDefaultLuminanceRange"
+	bool bShouldGenerateCode = false;
+	bool bCopyStarterContent = false;
+	
+	bool bForceExtendedLuminanceRange = false; // See "r.DefaultFeature.AutoExposure.ExtendDefaultLuminanceRange"
+	bool bEnableVR = false;
+	bool bIsEnterpriseProject = false;
 
-	EHardwareClass::Type TargetedHardware;
-	EGraphicsPreset::Type DefaultGraphicsPerformance;
+	EHardwareClass::Type TargetedHardware = EHardwareClass::Desktop;
+	EGraphicsPreset::Type DefaultGraphicsPerformance = EGraphicsPreset::Maximum;
 };
 
 DECLARE_DELEGATE_RetVal_OneParam(bool, FProjectDescriptorModifier, FProjectDescriptor&);
@@ -135,6 +129,9 @@ public:
 	/** Adds new source code to the project. When returning Succeeded or FailedToHotReload, OutSyncFileAndLineNumber will be the the preferred target file to sync in the users code editing IDE, formatted for use with GenericApplication::GotoLineInSource */
 	static EAddCodeToProjectResult AddCodeToProject(const FString& NewClassName, const FString& NewClassPath, const FModuleContextInfo& ModuleInfo, const FNewClassInfo ParentClassInfo, const TSet<FString>& DisallowedHeaderNames, FString& OutHeaderFilePath, FString& OutCppFilePath, FText& OutFailReason);
 
+	/** Loads a list of template categories defined in the TemplateCategories.ini file in the specified folder */
+	static UTemplateCategories* LoadTemplateCategories(const FString& RootDir);
+
 	/** Loads a template project definitions object from the TemplateDefs.ini file in the specified project */
 	static UTemplateProjectDefs* LoadTemplateDefs(const FString& ProjectDirectory);
 
diff --git a/Templates/TemplateCategories.ini b/Templates/TemplateCategories.ini
new file mode 100644
--- /dev/null
+++ b/Templates/TemplateCategories.ini
@@ -0,0 +1,7 @@
+[/Script/GameProjectGeneration.TemplateCategories]
+
+Categories=(Key="Game", LocalizedDisplayNames=((Language="en", Text="Game")), LocalizedDescriptions=((Language="en", Text="Making fancy games with the running and the questing and looting with the juicy inventories and UIs.")), Icon="GameProjectDialog.CodeIcon", Image="GameProjectDialog.CodeImage", IsMajorCategory=true)
+
+Categories=(Key="Virtual Production", LocalizedDisplayNames=((Language="en", Text="Virtual Production")), LocalizedDescriptions=((Language="en", Text="Movies with all the power of the Unreal and doodads like nDisplay and the compositions.")), Icon="GameProjectDialog.BlueprintIcon", Image="GameProjectDialog.BlueprintImage", IsMajorCategory=false)
+
+Categories=(Key="Visualization", LocalizedDisplayNames=((Language="en", Text="Visualization")), LocalizedDescriptions=((Language="en", Text="Pretty pretty buildings and cars and all that.")), Icon="GameProjectDialog.BlueprintIcon", Image="GameProjectDialog.BlueprintImage", IsMajorCategory=false)
\ No newline at end of file
