Project dialog has been rewritten to be more clearly split into a wizard-style flow.

diff --git a/Engine/Source/Editor/EditorStyle/Private/SlateEditorStyle.cpp b/Engine/Source/Editor/EditorStyle/Private/SlateEditorStyle.cpp
--- a/Engine/Source/Editor/EditorStyle/Private/SlateEditorStyle.cpp
+++ b/Engine/Source/Editor/EditorStyle/Private/SlateEditorStyle.cpp
@@ -1890,9 +1890,11 @@ void FSlateEditorStyle::FStyle::SetupGeneralStyles()
 		Set( "GameProjectDialog.BasicCodeThumbnail", new IMAGE_BRUSH( "GameProjectDialog/basic_code_thumbnail", Icon128x128 ) );
 		Set( "GameProjectDialog.CodeIcon", new IMAGE_BRUSH( "GameProjectDialog/feature_code_32x", FVector2D(32,32) ) );
 		Set( "GameProjectDialog.CodeImage", new IMAGE_BRUSH( "GameProjectDialog/feature_code", FVector2D(96,96) ) );
+		Set( "GameProjectDialog.CodeImage_64", new IMAGE_BRUSH( "GameProjectDialog/feature_code", Icon64x64) );
 		Set( "GameProjectDialog.CodeImage_48x", new IMAGE_BRUSH("GameProjectDialog/feature_code", Icon40x40));
 		Set( "GameProjectDialog.BlueprintIcon", new IMAGE_BRUSH( "GameProjectDialog/feature_blueprint_32x", FVector2D(32,32) ) );
 		Set( "GameProjectDialog.BlueprintImage", new IMAGE_BRUSH( "GameProjectDialog/feature_blueprint", FVector2D(96,96) ) );
+		Set( "GameProjectDialog.BlueprintImage_64", new IMAGE_BRUSH( "GameProjectDialog/feature_blueprint", Icon64x64) );
 		Set( "GameProjectDialog.CodeBorder", new BOX_BRUSH( "GameProjectDialog/feature_border", FMargin(4.0f/16.0f), FLinearColor(0.570, 0.359, 0.081, 1.f) ) );
 		Set( "GameProjectDialog.FeatureText", FTextBlockStyle(NormalText)
 			.SetFont( DEFAULT_FONT( "Regular", 14 ) )
diff --git a/Engine/Source/Editor/GameProjectGeneration/Private/SGameProjectDialog.cpp b/Engine/Source/Editor/GameProjectGeneration/Private/SGameProjectDialog.cpp
--- a/Engine/Source/Editor/GameProjectGeneration/Private/SGameProjectDialog.cpp
+++ b/Engine/Source/Editor/GameProjectGeneration/Private/SGameProjectDialog.cpp
@@ -6,17 +6,24 @@
 #include "Widgets/SOverlay.h"
 #include "Styling/CoreStyle.h"
 #include "SlateOptMacros.h"
-#include "Widgets/Layout/SBorder.h"
+#include "Widgets/Input/SButton.h"
+#include "Widgets/Input/SHyperlink.h"
 #include "Widgets/Images/SImage.h"
+#include "Widgets/Layout/SBorder.h"
 #include "Widgets/Layout/SBox.h"
-#include "Widgets/Input/SButton.h"
+#include "Widgets/Layout/SSeparator.h"
+#include "Widgets/Layout/SUniformGridPanel.h"
+#include "Widgets/Layout/SWidgetSwitcher.h"
 #include "EditorStyleSet.h"
 #include "GameProjectUtils.h"
 #include "SProjectBrowser.h"
 #include "SNewProjectWizard.h"
 #include "Widgets/SToolTip.h"
 #include "IDocumentation.h"
-#include "Widgets/Layout/SWidgetSwitcher.h"
+#include "Framework/MultiBox/MultiBoxBuilder.h"
+#include "Internationalization/BreakIterator.h"
+#include "TemplateCategory.h"
+#include "Framework/Application/SlateApplication.h"
 
 #define LOCTEXT_NAMESPACE "GameProjectGeneration"
 
@@ -24,280 +31,486 @@
 /* SGameProjectDialog interface
 *****************************************************************************/
 
-BEGIN_SLATE_FUNCTION_BUILD_OPTIMIZATION
-void SGameProjectDialog::Construct( const FArguments& InArgs )
+namespace GameProjectDialogDefs
 {
-	bool bAtLeastOneVisibleRecentProject = false;
+	constexpr float Padding = 5;
+	constexpr float TextWidth = 420;
+	constexpr float TextHeight = 16;
 
-	if (InArgs._AllowProjectCreate)
-	{
-		NewProjectWizard = SNew(SNewProjectWizard);
-	}
+	constexpr float ThumbnailSize = 64;
+	constexpr float MinorItemWidth = ThumbnailSize + TextWidth + Padding * 2;
+	constexpr float MinorItemHeight = ThumbnailSize + Padding * 3;
 
-	if (InArgs._AllowProjectOpening)
-	{
-		ProjectBrowser = SNew(SProjectBrowser);
-	}
+	constexpr float MajorItemWidth = MinorItemWidth * 2;
+	constexpr float MajorItemHeight = MinorItemHeight;
 
-	FadeAnimation.AddCurve( 0.f, 0.5f, ECurveEaseFunction::QuadOut );
-	RegisterActiveTimer( 0.f, FWidgetActiveTimerDelegate::CreateSP( this, &SGameProjectDialog::TriggerFadeInPostConstruct ) );
+	constexpr int32 LandingPageIndex = 0;
+	constexpr int32 ProjectBrowserPageIndex = 1;
+	constexpr int32 TemplateListPageIndex = 2;
+	constexpr int32 ProjectSettingsPageIndex = 3;
+	
+	static const FText LandingPageTitle = LOCTEXT("ProjectDialog_SelectProject", "Select Project");
+	static const FText SelectCategoryTitle = LOCTEXT("ProjectDialog_SelectCategory", "Select Category");
+	static const FText ProjectBrowserTitle = LOCTEXT("ProjectDialog_ProjectBrowser", "Project Browser");
+	static const FText TemplateListTitle = LOCTEXT("ProjectDialog_SelectTemplate", "Select Template");
+	static const FText ProjectSettingsTitle = LOCTEXT("ProjectDialog_ProjectSettings", "Project Settings");
+}
 
-	TSharedPtr<SWidget> Content;
+BEGIN_SLATE_FUNCTION_BUILD_OPTIMIZATION
+void SGameProjectDialog::Construct(const FArguments& InArgs, EMode InMode)
+{
+	DialogMode = InMode;
 
-	if (InArgs._AllowProjectCreate && InArgs._AllowProjectOpening)
-	{
-		// Create the Open Project tab button
-		TSharedRef<SButton> ProjectsTabButton = SNew(SButton)
-			.ForegroundColor(FCoreStyle::Get().GetSlateColor("Foreground"))
-			.ButtonStyle(FEditorStyle::Get(), TEXT("NoBorder"))
-			.OnClicked(this, &SGameProjectDialog::HandleProjectsTabButtonClicked)
-			.ContentPadding(FMargin(40, 5))
-			.Text(LOCTEXT("ProjectsTabTitle", "Projects"))
-			.TextStyle(FEditorStyle::Get(), TEXT("ProjectBrowser.Tab.Text"));
-
-		// Create the New Project tab button
-		TSharedRef<SButton> NewProjectTabButton = SNew(SButton)
-			.ForegroundColor(FCoreStyle::Get().GetSlateColor("Foreground"))
-			.ButtonStyle(FEditorStyle::Get(), "NoBorder")
-			.OnClicked(this, &SGameProjectDialog::HandleNewProjectTabButtonClicked)
-			.ContentPadding(FMargin(20, 5))
-			.TextStyle(FEditorStyle::Get(), "ProjectBrowser.Tab.Text")
-			.Text(LOCTEXT("NewProjectTabTitle", "New Project"))
-			.ToolTip(IDocumentation::Get()->CreateToolTip(LOCTEXT("NewProjectTabTitle", "New Project"), nullptr, "Shared/LevelEditor", "NewProjectTab"));
-
-		// Allow creation and opening, so we need tabs here
-		ChildSlot
-		[
-		
-			SNew(SBorder)
-			.ColorAndOpacity(this, &SGameProjectDialog::HandleCustomContentColorAndOpacity)
-			.BorderImage(FEditorStyle::GetBrush("Docking.Tab.ContentAreaBrush"))
-			.Padding(0)
-			[
-				SNew(SVerticalBox)
+	NewProjectWizard = SNew(SNewProjectWizard)
+		.OnTemplateDoubleClick(this, &SGameProjectDialog::OnTemplateDoubleClick);
+	ProjectSettingsPage = NewProjectWizard->CreateProjectSettingsPage();
+	ProjectBrowserPage = SNew(SProjectBrowser);
+	LandingPage = CreateLandingPage();
 
-				+ SVerticalBox::Slot()
-				.AutoHeight()
-				.Padding(FMargin(6.f, 0, 0, 0))
-				[
-					SNew(SHorizontalBox)
+	const float UniformPadding = 16.0f;
 
-					// Open Project Tab
-					+ SHorizontalBox::Slot()
-					.AutoWidth()
-					[
-						SNew( SBorder )
-						.BorderImage(this, &SGameProjectDialog::OnGetTabBorderImage, ProjectsTab)
-						.Padding(0)
-						[
-							SNew(SOverlay)
-
-							+ SOverlay::Slot()
-							.VAlign(VAlign_Top)
-							[
-								SNew(SBox)
-								.HeightOverride(2.0f)
-								[
-									SNew(SImage)
-									.Image(this, &SGameProjectDialog::OnGetTabHeaderImage, ProjectsTab, ProjectsTabButton)
-									.Visibility(EVisibility::HitTestInvisible)
-								]
-							]
-
-							+ SOverlay::Slot()
-							[
-								ProjectsTabButton
-							]
-						]
-					]
+	ChildSlot
+	[
+		SNew(SBorder)
+		.BorderImage(FEditorStyle::GetBrush("ToolPanel.DarkGroupBorder"))
+		.Padding(8)
+		[
+			SAssignNew(RootWizard, SWizard)
+			.ShowBreadcrumbs(false)
+			.ShowPageList(false)
+			.ShowPageTitle(true)
+			.ButtonStyle(FEditorStyle::Get(), "FlatButton.Default")
+			.CancelButtonStyle(FEditorStyle::Get(), "FlatButton.Default")
+			.FinishButtonStyle(FEditorStyle::Get(), "FlatButton.Success")
+			.ButtonTextStyle(FEditorStyle::Get(), "LargeText")
+			.ForegroundColor(FEditorStyle::Get().GetSlateColor("WhiteBrush"))
+			.FinishButtonText(this, &SGameProjectDialog::GetFinishText)
+			.FinishButtonToolTip(this, &SGameProjectDialog::GetFinishTooltip)
+			.CanFinish(this, &SGameProjectDialog::OnCanFinish)
+			.OnFinished(this, &SGameProjectDialog::OnFinishClicked)
+			.OnCanceled(this, &SGameProjectDialog::OnCancelClicked)
+			.InitialPageIndex(this, &SGameProjectDialog::GetInitialPageIndex)
+			.OnGetNextPageIndex(this, &SGameProjectDialog::GetNextPageIndex)
+
+			+ SWizard::Page()
+			.Name(GetPageTitle(GameProjectDialogDefs::LandingPageIndex))
+			.CanShow(DialogMode != EMode::Open)
+			[
+				LandingPage.ToSharedRef()
+			]
 
-					+ SHorizontalBox::Slot()
-					.Padding(FMargin(6.f, 0, 0, 0))
-					.AutoWidth()
-					[
-						SNew(SBorder)
-						.BorderImage(this, &SGameProjectDialog::OnGetTabBorderImage, NewProjectTab)
-						.Padding(0)
-						[
-							SNew(SOverlay)
-
-							+ SOverlay::Slot()
-							.VAlign(VAlign_Top)
-							[
-								SNew(SBox)
-								.HeightOverride(2.0f)
-								[
-									SNew(SImage)
-									.Image(this, &SGameProjectDialog::OnGetTabHeaderImage, NewProjectTab, NewProjectTabButton)
-									.Visibility(EVisibility::HitTestInvisible)
-								]
-							]
-
-							+ SOverlay::Slot()
-							[
-								NewProjectTabButton
-							]
-						]
-					]
-				]
+			+ SWizard::Page()
+			.Name(GetPageTitle(GameProjectDialogDefs::ProjectBrowserPageIndex))
+			.CanShow(DialogMode != EMode::New)
+			[
+				ProjectBrowserPage.ToSharedRef()
+			]
+		
+			+ SWizard::Page()
+			.Name(GetPageTitle(GameProjectDialogDefs::TemplateListPageIndex))
+			.CanShow(DialogMode != EMode::Open)
+			[
+				NewProjectWizard.ToSharedRef()
+			]
 
-				+ SVerticalBox::Slot()
-				[
-					SAssignNew(ContentAreaSwitcher, SWidgetSwitcher)
-					.WidgetIndex(0)
+			+ SWizard::Page()
+			.Name(GetPageTitle(GameProjectDialogDefs::ProjectSettingsPageIndex))
+			.CanShow(DialogMode != EMode::Open)
+			[
+				ProjectSettingsPage.ToSharedRef()
+			]
+		]
+	];
+}
 
-					+ SWidgetSwitcher::Slot()
-					[
-						ProjectBrowser.ToSharedRef()
-					]
+TSharedRef<SWidget> SGameProjectDialog::CreateLandingPage()
+{
+	TArray<TSharedPtr<FTemplateCategory>> AllTemplateCategories;
+	FGameProjectGenerationModule::Get().GetAllTemplateCategories(AllTemplateCategories);
 
-					+ SWidgetSwitcher::Slot()
-					[
-						NewProjectWizard.ToSharedRef()
-					]
-				]
-			]
-		];
-	}
-	else
+	for (const TSharedPtr<FTemplateCategory>& Category : AllTemplateCategories)
 	{
-		TSharedPtr<SWidget> BorderContent;
-		if (NewProjectWizard.IsValid())
+		if (Category->IsMajor)
 		{
-			BorderContent = NewProjectWizard;
+			MajorTemplateCategories.Add(Category);
 		}
 		else
 		{
-			BorderContent = ProjectBrowser;
+			MinorTemplateCategories.Add(Category);
 		}
-
-		ChildSlot
-		[
-			BorderContent.ToSharedRef()
-		];
 	}
 
-	ActiveTab = InArgs._AllowProjectOpening ? ProjectsTab : NewProjectTab;
-	if (ProjectBrowser.IsValid())
-	{
-		ActiveTab = !ProjectBrowser->HasProjects() && InArgs._AllowProjectCreate ? NewProjectTab : ActiveTab;
-	}
+	FText Description;
 
-	if (ActiveTab == ProjectsTab)
+	if (DialogMode == EMode::New)
 	{
-		ShowProjectBrowser();
+		Description = LOCTEXT("LandingPageDesc_New", "Choose a template <RichTextBlock.BoldHighlight>category</> for your new project.");
 	}
-	else if (ActiveTab == NewProjectTab) //-V547
+	else if (DialogMode == EMode::Open)
 	{
-		ShowNewProjectTab();
+		Description = LOCTEXT("LandingPageDesc_Open", "Choose a recent <RichTextBlock.BoldHighlight>project</> to load, or press <RichTextBlock.BoldHighlight>More</> to see all projects on your computer.");
 	}
-	else
+	else if (DialogMode == EMode::Both)
 	{
-		check(false);
+		Description = LOCTEXT("LandingPageDesc_Open", "Choose a recent <RichTextBlock.BoldHighlight>project</> to load, press <RichTextBlock.BoldHighlight>More</> to see all projects on your computer, or choose a template <RichTextBlock.BoldHighlight>category</> to create a new project.");
 	}
-}
-END_SLATE_FUNCTION_BUILD_OPTIMIZATION
 
+	TSharedRef<SWidget> Widget = SNew(SVerticalBox)
+	+ SVerticalBox::Slot()
+	.AutoHeight()
+	.Padding(0, 9, 0, 8)
+	[
+		SNew(SRichTextBlock)
+		.Text(Description)
+		.AutoWrapText(true)
+		.DecoratorStyleSet(&FEditorStyle::Get())
+	]
+	+ SVerticalBox::Slot()
+	.FillHeight(1)
+	[
+		SNew(SBorder)
+		.BorderImage(FEditorStyle::GetBrush("ToolPanel.GroupBorder"))
+		[
+			SNew(SScrollBox)
+			+ SScrollBox::Slot()
+			[
+				SNew(SVerticalBox)
+				+ SVerticalBox::Slot()
+				.AutoHeight()
+				[
+					SNew(SOverlay)
+					.Visibility(DialogMode != EMode::New ? EVisibility::Visible : EVisibility::Collapsed )
+					+ SOverlay::Slot()
+					[
+						SAssignNew(RecentProjectBrowser, SRecentProjectBrowser)
+						.OnSelectionChanged(this, &SGameProjectDialog::OnRecentProjectSelectionChanged)
+					]
+					+ SOverlay::Slot()
+					.HAlign(HAlign_Right)
+					.VAlign(VAlign_Bottom)
+					[
+						SNew(SButton)
+						.Text(LOCTEXT("ProjectDialog_More", "More"))
+						.TextStyle(FEditorStyle::Get(), "LargeText")
+						.ButtonStyle(FEditorStyle::Get(), "FlatButton.Default")
+						.ForegroundColor(FEditorStyle::Get().GetSlateColor("WhiteBrush"))
+						.HAlign(HAlign_Center)
+						.VAlign(VAlign_Center)
+						.ContentPadding(FCoreStyle::Get().GetMargin("StandardDialog.ContentPadding"))
+						.OnClicked(this, &SGameProjectDialog::OnMoreProjectsClicked)
+					]
+				]
+				+ SVerticalBox::Slot()
+				.AutoHeight()
+				.Padding(0, 4, 0, 2)
+				[
+					SNew(SSeparator)
+					.Visibility(DialogMode == EMode::Both ? EVisibility::Visible : EVisibility::Collapsed)
+				]
+				+ SVerticalBox::Slot()
+				.Padding(8)
+				.AutoHeight()
+				[
+					SNew(STextBlock)
+					.TextStyle(FEditorStyle::Get(), "GameProjectDialog.ProjectNamePathLabels")
+					.Text(LOCTEXT("ProjectDialog_Categories", "Categories"))
+				]
 
-EActiveTimerReturnType SGameProjectDialog::TriggerFadeInPostConstruct(double InCurrentTime, float InDeltaTime)
-{
-	// Play the intro fade in the first frame after the widget is created.
-	// We start it now instead of Construct because there is a lot of elapsed time between Construct and when we
-	// see the dialog and the beginning of the animation is cut off.
-	FadeAnimation.Play(this->AsShared());
+				+ SVerticalBox::Slot()
+				.AutoHeight()
+				.Padding(8, 8, 8, 0)
+				[
+					SAssignNew(MajorCategoryTileView, STileView< TSharedPtr<FTemplateCategory> >)
+					.ListItemsSource(&MajorTemplateCategories)
+					.SelectionMode(ESelectionMode::Single)
+					.ClearSelectionOnClick(true)
+					.OnGenerateTile(this, &SGameProjectDialog::ConstructTile)
+					.ItemHeight(GameProjectDialogDefs::MajorItemHeight)
+					.ItemWidth(GameProjectDialogDefs::MajorItemWidth)
+					.OnMouseButtonDoubleClick(this, &SGameProjectDialog::OnTemplateCategoryDoubleClick)
+					.OnSelectionChanged(this, &SGameProjectDialog::OnMajorTemplateCategorySelectionChanged)
+				]
+				+ SVerticalBox::Slot()
+				.AutoHeight()
+				.Padding(8, 0, 8, 8)
+				[
+					SAssignNew(MinorCategoryTileView, STileView< TSharedPtr<FTemplateCategory> >)
+					.ListItemsSource(&MinorTemplateCategories)
+					.SelectionMode(ESelectionMode::Single)
+					.ClearSelectionOnClick(true)
+					.OnGenerateTile(this, &SGameProjectDialog::ConstructTile)
+					.ItemHeight(GameProjectDialogDefs::MinorItemHeight)
+					.ItemWidth(GameProjectDialogDefs::MinorItemWidth)
+					.OnMouseButtonDoubleClick(this, &SGameProjectDialog::OnTemplateCategoryDoubleClick)
+					.OnSelectionChanged(this, &SGameProjectDialog::OnMinorTemplateCategorySelectionChanged)
+				]
+			]
+		]
+	];
 
-	return EActiveTimerReturnType::Stop;
+	return Widget;
 }
 
-/* SGameProjectDialog implementation
+END_SLATE_FUNCTION_BUILD_OPTIMIZATION
+
+/* SGameProjectDialog callbacks
 *****************************************************************************/
 
-bool SGameProjectDialog::OpenProject( const FString& ProjectFile )
+TSharedRef<ITableRow> SGameProjectDialog::ConstructTile(TSharedPtr<FTemplateCategory> Item, const TSharedRef<STableViewBase>& TableView)
 {
-	FText FailReason;
+	const float TextWidth = (Item->IsMajor ? GameProjectDialogDefs::TextWidth * 2 : GameProjectDialogDefs::TextWidth)/* - GameProjectDialogDefs::Padding * 2*/;
+
+	TSharedRef<STableRow<TSharedPtr<FTemplateCategory>>> Row = SNew(STableRow<TSharedPtr<FTemplateCategory>>, TableView)
+	.Style(FEditorStyle::Get(), "GameProjectDialog.TemplateListView.TableRow")
+	[
+		SNew(SScrollBox)
+		+ SScrollBox::Slot()
+		.Padding(0)
+		[
+			SNew(SHorizontalBox)
+
+			// Icon
+			+ SHorizontalBox::Slot()
+			.AutoWidth()
+			.HAlign(HAlign_Center)
+			.VAlign(VAlign_Center)
+			.Padding(GameProjectDialogDefs::Padding)
+			[
+				SNew(SBox)
+				.VAlign(VAlign_Center)
+				.HAlign(HAlign_Center)
+				.WidthOverride(GameProjectDialogDefs::ThumbnailSize)
+				.HeightOverride(GameProjectDialogDefs::ThumbnailSize)
+				.Padding(0)
+				[
+					SNew(SImage)
+					.Image(Item->Image)
+				]
+			]
 
-	if (GameProjectUtils::OpenProject(ProjectFile, FailReason))
-	{
-		return true;
-	}
+			+ SHorizontalBox::Slot()
+			.Padding(GameProjectDialogDefs::Padding)
+			.VAlign(VAlign_Fill)
+			[
+				SNew(SVerticalBox)
+				// Name
+				+ SVerticalBox::Slot()
+				.MaxHeight(GameProjectDialogDefs::TextHeight)
+				.Padding(0, 0, 0, GameProjectDialogDefs::Padding)
+				[
+					SNew(STextBlock)
+					.Text(Item->DisplayName)
+					.Justification(ETextJustify::Left)
+					.TextStyle(FEditorStyle::Get(), "LargeText")
+				]
 
-	FMessageDialog::Open(EAppMsgType::Ok, FailReason);
+				// Description
+				+ SVerticalBox::Slot()
+				.AutoHeight()
+				.Padding(0)
+				[
+					SNew(STextBlock)
+					.WrapTextAt(TextWidth)
+					.LineBreakPolicy(FBreakIterator::CreateWordBreakIterator())
+					.Justification(ETextJustify::Left)
+					.Text(Item->Description)
+				]
+			]
+		]
+	];
 
-	return false;
+	Row->SetToolTipText(FText::Format(FText::FromString("{0}\n{1}"), Item->DisplayName, Item->Description));
+
+	return Row;
 }
 
+FReply SGameProjectDialog::OnMoreProjectsClicked()
+{
+	int32 Index = GameProjectDialogDefs::ProjectBrowserPageIndex;
+	RootWizard->AdvanceToPage(Index);
+	return FReply::Handled();
+}
 
-void SGameProjectDialog::ShowNewProjectTab( )
+void SGameProjectDialog::OnTemplateCategoryDoubleClick(TSharedPtr<FTemplateCategory> Item) const
 {
-	if (ContentAreaSwitcher.IsValid() && NewProjectWizard.IsValid())
-	{
-		ContentAreaSwitcher->SetActiveWidget( NewProjectWizard.ToSharedRef());
-		ActiveTab = NewProjectTab;
-	}
+	RootWizard->AdvanceToPage(GameProjectDialogDefs::TemplateListPageIndex);
 }
 
+void SGameProjectDialog::OnTemplateDoubleClick() const
+{
+	RootWizard->AdvanceToPage(GameProjectDialogDefs::ProjectSettingsPageIndex);
+}
 
-FReply SGameProjectDialog::ShowProjectBrowser( )
+void SGameProjectDialog::OnMajorTemplateCategorySelectionChanged(TSharedPtr<FTemplateCategory> Item, ESelectInfo::Type SelectType)
 {
-	if (ContentAreaSwitcher.IsValid() && ProjectBrowser.IsValid())
+	if (Item.IsValid())
 	{
-		ContentAreaSwitcher->SetActiveWidget(ProjectBrowser.ToSharedRef());
-		ActiveTab = ProjectsTab;
+		RecentProjectBrowser->ClearSelection();
+		MinorCategoryTileView->ClearSelection();
+
+		NewProjectWizard->SetCurrentCategory(Item->Key);
 	}
 
-	return FReply::Handled();
 }
 
+void SGameProjectDialog::OnMinorTemplateCategorySelectionChanged(TSharedPtr<FTemplateCategory> Item, ESelectInfo::Type SelectType)
+{
+	if (Item.IsValid())
+	{
+		RecentProjectBrowser->ClearSelection();
+		MajorCategoryTileView->ClearSelection();
 
-/* SGameProjectDialog callbacks
-*****************************************************************************/
+		NewProjectWizard->SetCurrentCategory(Item->Key);
+	}
+}
 
-FLinearColor SGameProjectDialog::HandleCustomContentColorAndOpacity( ) const
+void SGameProjectDialog::OnRecentProjectSelectionChanged(FString Item)
 {
-	return FLinearColor(1,1,1, FadeAnimation.GetLerp());
+	if (!Item.IsEmpty())
+	{
+		MajorCategoryTileView->ClearSelection();
+		MinorCategoryTileView->ClearSelection();
+	}
 }
 
+int32 SGameProjectDialog::GetInitialPageIndex() const
+{
+	if (DialogMode == EMode::Open)
+	{
+		return GameProjectDialogDefs::ProjectBrowserPageIndex;
+	}
+	return GameProjectDialogDefs::LandingPageIndex;
+}
 
-FReply SGameProjectDialog::HandleNewProjectTabButtonClicked( )
+int32 SGameProjectDialog::GetNextPageIndex(int32 Current) const
 {
-	ShowNewProjectTab();
+	int32 CurrentPage = RootWizard->GetCurrentPageIndex();
+	if (CurrentPage == GameProjectDialogDefs::LandingPageIndex)
+	{
+		FString SelectedProjectFile = RecentProjectBrowser->GetSelectedProjectFile();
+		if (!SelectedProjectFile.IsEmpty())
+		{
+			return INDEX_NONE;
+		}
 
-	return FReply::Handled();
-}
+		TArray<TSharedPtr<FTemplateCategory>> SelectedMajorCategories = MajorCategoryTileView->GetSelectedItems();
+		TArray<TSharedPtr<FTemplateCategory>> SelectedMinorCategories = MinorCategoryTileView->GetSelectedItems();
 
+		if (SelectedMajorCategories.Num() > 0 || SelectedMinorCategories.Num() > 0)
+		{
+			return GameProjectDialogDefs::TemplateListPageIndex;
+		}
+	}
+	else if (CurrentPage == GameProjectDialogDefs::TemplateListPageIndex)
+	{
+		if (NewProjectWizard->ShouldShowProjectSettingsPage())
+		{
+			return GameProjectDialogDefs::ProjectSettingsPageIndex;
+		}
+	}
+
+	return INDEX_NONE;
+}
 
-const FSlateBrush* SGameProjectDialog::OnGetTabHeaderImage( ETab InTab, TSharedRef<SButton> TabButton ) const
+bool SGameProjectDialog::OnCanFinish() const
 {
-	if (TabButton->IsPressed())
+	int32 CurrentPage = RootWizard->GetCurrentPageIndex();
+	if (CurrentPage == GameProjectDialogDefs::ProjectSettingsPageIndex)
 	{
-		return FEditorStyle::GetBrush("ProjectBrowser.Tab.PressedHighlight");
+		return NewProjectWizard->CanCreateProject();
 	}
 
-	if ((ActiveTab == InTab) || TabButton->IsHovered())
+	if (CurrentPage == GameProjectDialogDefs::TemplateListPageIndex)
 	{
-		return FEditorStyle::GetBrush("ProjectBrowser.Tab.ActiveHighlight");
+		if (!NewProjectWizard->ShouldShowProjectSettingsPage()) 
+		{
+			return NewProjectWizard->CanCreateProject();
+		}
 	}
-	
-	return FEditorStyle::GetBrush("ProjectBrowser.Tab.Highlight");
+
+	if (CurrentPage == GameProjectDialogDefs::ProjectBrowserPageIndex)
+	{
+		return !ProjectBrowserPage->GetSelectedProjectFile().IsEmpty();
+	}
+
+	if (CurrentPage == GameProjectDialogDefs::LandingPageIndex)
+	{
+		if (DialogMode == EMode::New)
+		{
+			return false;
+		}
+		else
+		{
+			return !RecentProjectBrowser->GetSelectedProjectFile().IsEmpty();
+		}
+	}
+
+	return false;
 }
 
+void SGameProjectDialog::OnFinishClicked()
+{
+	int32 CurrentPage = RootWizard->GetCurrentPageIndex();
+	if (CurrentPage == GameProjectDialogDefs::LandingPageIndex)
+	{
+		RecentProjectBrowser->OpenSelectedProject();
+	}
+	else if (CurrentPage == GameProjectDialogDefs::ProjectBrowserPageIndex)
+	{
+		ProjectBrowserPage->OpenSelectedProject();
+	}
+	else
+	{
+		NewProjectWizard->CreateAndOpenProject();
+	}
+
+	TSharedPtr<SWindow> Window = FSlateApplication::Get().FindWidgetWindow(AsShared());
+	Window->RequestDestroyWindow();
+}
 
-FReply SGameProjectDialog::HandleProjectsTabButtonClicked( )
+void SGameProjectDialog::OnCancelClicked() const
 {
-	ShowProjectBrowser();
+	TSharedPtr<SWindow> Window = FSlateApplication::Get().FindWidgetWindow(AsShared());
+	Window->RequestDestroyWindow();
+}
 
-	return FReply::Handled();
+FText SGameProjectDialog::GetFinishText() const
+{
+	int32 CurrentPage = RootWizard->GetCurrentPageIndex();
+	if ((DialogMode == EMode::Open || DialogMode == EMode::Both) &&
+		(CurrentPage == GameProjectDialogDefs::ProjectBrowserPageIndex ||
+		CurrentPage == GameProjectDialogDefs::LandingPageIndex))
+	{
+		return NSLOCTEXT("ProjectBrowser", "OpenProjectBrowseTitle", "Open Project");
+	}
+
+	return NSLOCTEXT("NewProjectWizard", "Create Project", "Create Project");
 }
 
-const FSlateBrush* SGameProjectDialog::OnGetTabBorderImage( ETab InTab ) const
+FText SGameProjectDialog::GetFinishTooltip() const
 {
-	if (ActiveTab == InTab)
+	int32 CurrentPage = RootWizard->GetCurrentPageIndex();
+	if (CurrentPage == GameProjectDialogDefs::TemplateListPageIndex ||
+		CurrentPage == GameProjectDialogDefs::ProjectSettingsPageIndex)
 	{
-		return FEditorStyle::GetBrush("ProjectBrowser.Tab.ActiveBackground");
+		return NSLOCTEXT("NewProjectWizard", "FinishButtonToolTip", "Creates your new project in the specified location with the specified template and name.");
 	}
-	
-	return FEditorStyle::GetBrush("ProjectBrowser.Tab.Background");
+
+	return FText::GetEmpty();
 }
 
+FText SGameProjectDialog::GetPageTitle(int32 PageIndex) const
+{
+	switch (PageIndex)
+	{
+		case GameProjectDialogDefs::LandingPageIndex:
+			return DialogMode == EMode::New ? GameProjectDialogDefs::SelectCategoryTitle : GameProjectDialogDefs::LandingPageTitle;
+		case GameProjectDialogDefs::ProjectBrowserPageIndex: return GameProjectDialogDefs::ProjectBrowserTitle;
+		case GameProjectDialogDefs::TemplateListPageIndex: return GameProjectDialogDefs::TemplateListTitle;
+		case GameProjectDialogDefs::ProjectSettingsPageIndex: return GameProjectDialogDefs::ProjectSettingsTitle;
+	}
+
+	return FText::GetEmpty();
+}
 
 #undef LOCTEXT_NAMESPACE
diff --git a/Engine/Source/Editor/GameProjectGeneration/Private/SGameProjectDialog.h b/Engine/Source/Editor/GameProjectGeneration/Private/SGameProjectDialog.h
--- a/Engine/Source/Editor/GameProjectGeneration/Private/SGameProjectDialog.h
+++ b/Engine/Source/Editor/GameProjectGeneration/Private/SGameProjectDialog.h
@@ -8,11 +8,17 @@
 #include "Input/Reply.h"
 #include "Animation/CurveSequence.h"
 #include "Widgets/SCompoundWidget.h"
+#include "Widgets/Navigation/SBreadcrumbTrail.h"
+#include "Widgets/Views/STileView.h"
+#include "Widgets/Workflow/SWizard.h"
 
 class SButton;
 class SNewProjectWizard;
 class SProjectBrowser;
 struct FSlateBrush;
+class STableViewBase;
+struct FTemplateCategory;
+class SRecentProjectBrowser;
 
 /**
  * A dialog to create a new project or open an existing one
@@ -20,73 +26,61 @@ struct FSlateBrush;
 class SGameProjectDialog
 	: public SCompoundWidget
 {
-	// Enumerates available tabs.
-	enum ETab
-	{
-		ProjectsTab,
-		NewProjectTab
-	};
-
 public:
 
-	SLATE_BEGIN_ARGS(SGameProjectDialog) { }
-		SLATE_ARGUMENT(bool, AllowProjectOpening)
-		SLATE_ARGUMENT(bool, AllowProjectCreate)
+	SLATE_BEGIN_ARGS(SGameProjectDialog) {}
 	SLATE_END_ARGS()
 
 public:
 
+	enum class EMode
+	{
+		Open,
+		New,
+		Both
+	};
+
 	/** Constructs this widget with InArgs */
-	void Construct( const FArguments& InArgs );
-
-protected:
-	/**
-	 * Opens the specified project.
-	 *
-	 * @param ProjectFile - The project file to open.
-	 *
-	 * @return true if the project was opened, false otherwise.
-	 */
-	bool OpenProject( const FString& ProjectFile );
-
-	/**
-	 * Shows the 'New Project' tab.
-	 */
-	void ShowNewProjectTab( );
-
-	/**
-	 * Shows the project browser tab.
-	 */
-	FReply ShowProjectBrowser( );
+	void Construct( const FArguments& InArgs, EMode Mode );
 
 private:
-	/** Ensures the fade-in animation is played post-construct */
-	EActiveTimerReturnType TriggerFadeInPostConstruct( double InCurrentTime, float InDeltaTime );
+	EMode DialogMode;
 
-	// Callback for getting the color of the custom content area.
-	FLinearColor HandleCustomContentColorAndOpacity() const;
+	TSharedPtr<SWizard> RootWizard;
+	TSharedPtr<SProjectBrowser> ProjectBrowserPage;
+	TSharedPtr<SNewProjectWizard> NewProjectWizard;
+	TSharedPtr<SRecentProjectBrowser> RecentProjectBrowser;
+	TSharedPtr<SWidget> LandingPage;
+	TSharedPtr<SWidget> ProjectSettingsPage;
 
-	// Callback for clicking the 'New Project' button.
-	FReply HandleNewProjectTabButtonClicked( );
+	TSharedPtr<STileView<TSharedPtr<FTemplateCategory>>> MajorCategoryTileView;
+	TSharedPtr<STileView<TSharedPtr<FTemplateCategory>>> MinorCategoryTileView;
 
-	// Callback for clicking the 'Projects' button.
-	FReply HandleProjectsTabButtonClicked( );
+	TArray<TSharedPtr<FTemplateCategory>> MajorTemplateCategories;
+	TArray<TSharedPtr<FTemplateCategory>> MinorTemplateCategories;
 
-	// Callback for getting the border image of the specified tab.
-	const FSlateBrush* OnGetTabBorderImage( ETab InTab ) const;
+private:
 
-	// Callback for getting the header stripe image for the specified tab.
-	const FSlateBrush* OnGetTabHeaderImage( ETab InTab, TSharedRef<SButton> TabButton ) const;
+	TSharedRef<SWidget> CreateLandingPage();
 
-private:
+	TSharedRef<ITableRow> ConstructTile(TSharedPtr<FTemplateCategory> Item, const TSharedRef<STableViewBase>& TableView);
 
-	// Holds the fading animation.
-	FCurveSequence FadeAnimation;
+	FReply OnMoreProjectsClicked();
 
-	/** The switcher widget to control which screen is in view */
-	TSharedPtr<SWidgetSwitcher> ContentAreaSwitcher;
-	TSharedPtr<SProjectBrowser> ProjectBrowser;
-	TSharedPtr<SNewProjectWizard> NewProjectWizard;
+	void OnTemplateCategoryDoubleClick(TSharedPtr<FTemplateCategory> Item) const;
+	void OnTemplateDoubleClick() const;
+
+	void OnMajorTemplateCategorySelectionChanged(TSharedPtr<FTemplateCategory> Item, ESelectInfo::Type SelectType);
+	void OnMinorTemplateCategorySelectionChanged(TSharedPtr<FTemplateCategory> Item, ESelectInfo::Type SelectType);
+	void OnRecentProjectSelectionChanged(FString Item);
+
+	int32 GetInitialPageIndex() const;
+	int32 GetNextPageIndex(int32 Current) const;
 
-	ETab ActiveTab;
+	void OnCancelClicked() const;
+	bool OnCanFinish() const;
+	void OnFinishClicked();
+	FText GetFinishText() const;
+	FText GetFinishTooltip() const;
+	FText GetPageTitle(int32 Index) const;
 };
diff --git a/Engine/Source/Editor/GameProjectGeneration/Private/SNewProjectWizard.cpp b/Engine/Source/Editor/GameProjectGeneration/Private/SNewProjectWizard.cpp
--- a/Engine/Source/Editor/GameProjectGeneration/Private/SNewProjectWizard.cpp
+++ b/Engine/Source/Editor/GameProjectGeneration/Private/SNewProjectWizard.cpp
@@ -17,11 +17,12 @@
 #include "Widgets/Layout/SSeparator.h"
 #include "Widgets/Images/SImage.h"
 #include "Widgets/Text/SRichTextBlock.h"
-#include "Widgets/Layout/SGridPanel.h"
 #include "Widgets/Input/SEditableTextBox.h"
 #include "Widgets/Input/SButton.h"
+#include "Widgets/Layout/SGridPanel.h"
 #include "Widgets/Layout/SScrollBorder.h"
 #include "Widgets/Layout/SScrollBox.h"
+#include "Widgets/Layout/SWrapBox.h"
 #include "Widgets/Input/SCheckBox.h"
 #include "EditorStyleSet.h"
 #include "Editor.h"
@@ -51,7 +52,6 @@
 FName SNewProjectWizard::TemplatePageName = TEXT("Template");
 FName SNewProjectWizard::NameAndLocationPageName = TEXT("NameAndLocation");
 
-
 namespace NewProjectWizardDefs
 {
 	const float ThumbnailSize = 64.f, ThumbnailPadding = 5.f;
@@ -315,11 +315,13 @@ void SNewProjectWizard::Construct( const FArguments& InArgs )
 	SelectedHardwareClassTarget = EHardwareClass::Desktop;
 	SelectedGraphicsPreset = EGraphicsPreset::Maximum;
 
+	OnTemplateDoubleClick = InArgs._OnTemplateDoubleClick;
+
 	// Find all template projects
 	FindTemplateProjects();
 	SetDefaultProjectLocation();
 
-	SAssignNew(TemplateListView, STileView< TSharedPtr<FTemplateItem> >)
+	TemplateListView = SNew(STileView< TSharedPtr<FTemplateItem> >)
 	.ListItemsSource(&FilteredTemplateList)
 	.SelectionMode(ESelectionMode::Single)
 	.ClearSelectionOnClick(false)
@@ -334,446 +336,196 @@ void SNewProjectWizard::Construct( const FArguments& InArgs )
 	TSharedRef<SSeparator> Separator = SNew(SSeparator).Orientation(EOrientation::Orient_Vertical);
 	Separator->SetBorderBackgroundColor(FLinearColor::White.CopyWithNewOpacity(0.25f));
 
-	const float UniformPadding = 16.f;
 	ChildSlot
+	.Padding(0)
 	[
-		SNew(SBorder)
-		.BorderImage( FEditorStyle::GetBrush("ToolPanel.GroupBorder") )
+		SNew(SVerticalBox)
+
+		+ SVerticalBox::Slot()
+		.Padding(0, 8, 0, 8)
+		.AutoHeight()
 		[
-			SNew(SOverlay)
+			SNew(SRichTextBlock)
+			.Text(LOCTEXT("ProjectTemplateDescription", "Choose a <RichTextBlock.BoldHighlight>template</> to use as a starting point for your new project.  Any of these features can be added later by clicking <RichTextBlock.BoldHighlight>Add Feature or Content Pack</> in <RichTextBlock.BoldHighlight>Content Browser</>."))
+			.AutoWrapText(true)
+			.DecoratorStyleSet(&FEditorStyle::Get())
+			.ToolTip(IDocumentation::Get()->CreateToolTip(LOCTEXT("TemplateChoiceTooltip", "A template consists of a little bit of player control logic (either as a Blueprint or in C++), input bindings, and appropriate prototyping assets."), NULL, TEXT("Shared/Editor/NewProjectWizard"), TEXT("TemplateChoice")))
+		]
+
+		+ SVerticalBox::Slot()
+		[
+			SNew(SVerticalBox)
 
-			// Wizard
-			+SOverlay::Slot()
-			.Padding(UniformPadding / 2.0f)
+			// Templates list
+			+ SVerticalBox::Slot()
+			.FillHeight(1.0f)
+			.Padding(0)
 			[
-				SAssignNew(MainWizard, SWizard)
-				.ButtonStyle(FEditorStyle::Get(), "FlatButton.Default")
-				.CancelButtonStyle(FEditorStyle::Get(), "FlatButton.Default")
-				.FinishButtonStyle(FEditorStyle::Get(), "FlatButton.Success")
-				.ButtonTextStyle(FEditorStyle::Get(), "LargeText")
-				.ForegroundColor(FEditorStyle::Get().GetSlateColor("WhiteBrush"))
-				.ShowPageList(false)
-				.ShowCancelButton(false)
-				.CanFinish(this, &SNewProjectWizard::HandleCreateProjectWizardCanFinish)
-				.FinishButtonText(LOCTEXT("FinishButtonText", "Create Project"))
-				.FinishButtonToolTip(LOCTEXT("FinishButtonToolTip", "Creates your new project in the specified location with the specified template and name.") )
-				.OnFinished(this, &SNewProjectWizard::HandleCreateProjectWizardFinished)
-				.OnFirstPageBackClicked(InArgs._OnBackRequested)
-
-				// Choose Template
-				+SWizard::Page()
-				.OnEnter(this, &SNewProjectWizard::OnPageVisited, TemplatePageName)
+				SNew(SBorder)
+				.Padding(8)
+				.BorderImage(FEditorStyle::GetBrush("ToolPanel.GroupBorder"))
 				[
-					SNew(SBorder)
-					.BorderImage(FEditorStyle::GetBrush("NoBorder"))
-					.Padding(FMargin(UniformPadding / 2.0f, UniformPadding / 2.0f, UniformPadding / 2.0f, 0.0f))
-					[
-						SNew(SVerticalBox)
+					SNew(SHorizontalBox)
 
-						+ SVerticalBox::Slot()
-						.Padding(FMargin(0, 0, 0, 15))
-						.AutoHeight()
+					+ SHorizontalBox::Slot()
+					[
+						SNew(SScrollBorder, TemplateListView.ToSharedRef())
 						[
-							SNew(SRichTextBlock)
-							.Text(LOCTEXT("ProjectTemplateDescription", "Choose a <RichTextBlock.BoldHighlight>template</> to use as a starting point for your new project.  Any of these features can be added later by clicking <RichTextBlock.BoldHighlight>Add Feature or Content Pack</> in <RichTextBlock.BoldHighlight>Content Browser</>."))
-							.AutoWrapText(true)
-							.DecoratorStyleSet(&FEditorStyle::Get())
-							.ToolTip(IDocumentation::Get()->CreateToolTip(LOCTEXT("TemplateChoiceTooltip", "A template consists of a little bit of player control logic (either as a Blueprint or in C++), input bindings, and appropriate prototyping assets."), NULL, TEXT("Shared/Editor/NewProjectWizard"), TEXT("TemplateChoice")))
+							TemplateListView.ToSharedRef()
 						]
+					]
+										
+					+ SHorizontalBox::Slot()
+					.Padding(8, 0.0f)
+					.AutoWidth()
+					[
+						Separator
+					]
 
-						+ SVerticalBox::Slot()
+					// Selected template details
+					+ SHorizontalBox::Slot()
+					[
+						SNew(SScrollBox)
+						+ SScrollBox::Slot()
+						.Padding(8, 0.0f)
 						[
-							// Template category tabs
 							SNew(SVerticalBox)
 
-							+SVerticalBox::Slot()
-							.Padding(FMargin(8.f, 0.f))
+							// Preview image
+							+ SVerticalBox::Slot()
 							.AutoHeight()
+							.HAlign(HAlign_Center)
+							.Padding(FMargin(0.0f, 0.0f, 0.0f, 15.f))
 							[
-								BuildCategoryTabs()
+								SNew(SBox)
+								.Visibility(this, &SNewProjectWizard::GetSelectedTemplatePreviewVisibility)
+								.WidthOverride(400)
+								.HeightOverride(200)
+								[
+									SNew(SBorder)
+									.Padding(FMargin(0.0f, 0.0f, 0.0f, 4.f))
+									.BorderImage(FEditorStyle::GetBrush("ContentBrowser.ThumbnailShadow"))
+									[
+										SNew(SImage)
+										.Image(this, &SNewProjectWizard::GetSelectedTemplatePreviewImage)
+									]
+								]
 							]
 
-							// Templates list
+							// Template Name
+							+ SVerticalBox::Slot()
+							.Padding(FMargin(0.0f, 0.0f, 0.0f, 10.0f))
+							.AutoHeight()
+							[
+								SNew(STextBlock)
+								.AutoWrapText(true)
+								.TextStyle(FEditorStyle::Get(), "GameProjectDialog.FeatureText")
+								.Text(this, &SNewProjectWizard::GetSelectedTemplateProperty<FText>, &FTemplateItem::Name)
+							]
+						
+							// Template Description
+							+ SVerticalBox::Slot()
+							[
+								SNew(STextBlock)
+								.AutoWrapText(true)
+								.Text(this, &SNewProjectWizard::GetSelectedTemplateProperty<FText>, &FTemplateItem::Description)
+							]
+											
+							// Asset types
 							+ SVerticalBox::Slot()
-							.FillHeight(1.0f)
+							.AutoHeight()
+							.Padding(FMargin(0.0f, 5.0f, 0.0f, 5.0f))
 							[
-								SNew(SBorder)
-								.Padding(UniformPadding)
-								.BorderImage(FEditorStyle::GetBrush("GameProjectDialog.TabBackground"))
+								SNew(SBox)
+								.Visibility(this, &SNewProjectWizard::GetSelectedTemplateAssetVisibility)
 								[
-									SNew(SHorizontalBox)
-
-									+ SHorizontalBox::Slot()
-									[
-										SNew(SScrollBorder, TemplateListView.ToSharedRef())
-										[
-											TemplateListView.ToSharedRef()
-										]
-									]
-										
-									+ SHorizontalBox::Slot()
-									.Padding(UniformPadding, 0.0f)
-									.AutoWidth()
+									SNew(SVerticalBox)
+									+ SVerticalBox::Slot()
 									[
-										Separator
+										SNew(STextBlock)
+										.TextStyle(FEditorStyle::Get(), "GameProjectDialog.FeatureText")
+										.Text(LOCTEXT("ProjectTemplateAssetTypes", "Asset Type References:"))
 									]
-
-									// Selected template details
-									+ SHorizontalBox::Slot()
+									+ SVerticalBox::Slot()
+									.AutoHeight()
 									[
-										SNew(SScrollBox)
-										+ SScrollBox::Slot()
-										.Padding(UniformPadding, 0.0f)
-										[
-											SNew(SVerticalBox)
-
-											// Preview image
-											+ SVerticalBox::Slot()
-											.AutoHeight()
-											.HAlign(HAlign_Center)
-											.Padding(FMargin(0.0f, 0.0f, 0.0f, 15.f))
-											[
-												SNew(SBox)
-												.Visibility(this, &SNewProjectWizard::GetSelectedTemplatePreviewVisibility)
-												.WidthOverride(400)
-												.HeightOverride(200)
-												[
-													SNew(SOverlay)
-							
-													+ SOverlay::Slot()
-													[
-														SNew(SBorder)
-														.Padding(FMargin(0.0f, 0.0f, 0.0f, 4.f))
-														.BorderImage(FEditorStyle::GetBrush("ContentBrowser.ThumbnailShadow"))
-														[
-															SNew(SImage)
-															.Image(this, &SNewProjectWizard::GetSelectedTemplatePreviewImage)
-														]
-													]
-
-													+ SOverlay::Slot()
-													.HAlign(HAlign_Right)
-													.VAlign(VAlign_Top)
-													.Padding(10.0f)
-													[
-														SNew(SBox)
-														.WidthOverride(48.0f)
-														.HeightOverride(48.0f)
-														[
-															SNew(SImage)
-															.Image(this, &SNewProjectWizard::GetSelectedTemplateTypeImage)
-														]
-													]
-												]
-											]
-
-											// Template Name
-											+ SVerticalBox::Slot()
-											.Padding(FMargin(0.0f, 0.0f, 0.0f, 10.0f))
-											.AutoHeight()
-											[
-												SNew(STextBlock)
-												.AutoWrapText(true)
-												.TextStyle(FEditorStyle::Get(), "GameProjectDialog.FeatureText")
-												.Text(this, &SNewProjectWizard::GetSelectedTemplateProperty<FText>, &FTemplateItem::Name)
-											]
-						
-											// Template Description
-											+ SVerticalBox::Slot()
-											[
-												SNew(STextBlock)
-												.AutoWrapText(true)
-												.Text(this, &SNewProjectWizard::GetSelectedTemplateProperty<FText>, &FTemplateItem::Description)
-											]
-											
-											// Asset types
-											+ SVerticalBox::Slot()
-											.AutoHeight()
-											.Padding(FMargin(0.0f, 5.0f, 0.0f, 5.0f))
-											[
-												SNew(SBox)
-												.Visibility(this, &SNewProjectWizard::GetSelectedTemplateAssetVisibility)
-												[
-													SNew(SVerticalBox)
-													+ SVerticalBox::Slot()
-													[
-														SNew(STextBlock)
-														.TextStyle(FEditorStyle::Get(), "GameProjectDialog.FeatureText")
-														.Text(LOCTEXT("ProjectTemplateAssetTypes", "Asset Type References:"))
-													]
-													+ SVerticalBox::Slot()
-													.AutoHeight()
-													[
-														SNew(STextBlock)
-														.AutoWrapText(true)
-														.Text(this, &SNewProjectWizard::GetSelectedTemplateAssetTypes)
-													]
-												]
-											]
-											// Class types
-											+ SVerticalBox::Slot()
-											.AutoHeight()
-											.Padding(FMargin(0.0f, 5.0f, 0.0f, 5.0f))
-											[
-												SNew(SBox)
-												.Visibility(this, &SNewProjectWizard::GetSelectedTemplateClassVisibility)
-												[
-													SNew(SVerticalBox)
-													+ SVerticalBox::Slot()
-													[
-														SNew(STextBlock)
-														.TextStyle(FEditorStyle::Get(), "GameProjectDialog.FeatureText")
-														.Text(LOCTEXT("ProjectTemplateClassTypes", "Class Type References:"))
-													]
-													+ SVerticalBox::Slot()
-														.AutoHeight()
-														[
-															SNew(STextBlock)
-															.AutoWrapText(true)
-															.Text(this, &SNewProjectWizard::GetSelectedTemplateClassTypes)
-														]
-												]
-
-											]
-										]
+										SNew(STextBlock)
+										.AutoWrapText(true)
+										.Text(this, &SNewProjectWizard::GetSelectedTemplateAssetTypes)
 									]
 								]
 							]
-						]
-
-						+ SVerticalBox::Slot()
-						.AutoHeight()
-						.Padding(FMargin(0.0f, 15.0f, 0.0f, 0.0f))
-						[
-							SNew(SScrollBox)
-
-							+ SScrollBox::Slot()
+							// Class types
+							+ SVerticalBox::Slot()
+							.AutoHeight()
+							.Padding(FMargin(0.0f, 5.0f, 0.0f, 5.0f))
 							[
-								
-								SNew(SVerticalBox)
-									
-								+ SVerticalBox::Slot()
-								.AutoHeight()
-								.Padding(FMargin(0, 0, 0, 15.f))
+								SNew(SBox)
+								.Visibility(this, &SNewProjectWizard::GetSelectedTemplateClassVisibility)
 								[
-									MakeProjectSettingsDescriptionBox()
-								]
-
-								+ SVerticalBox::Slot()
-								.HAlign(HAlign_Center)
-								.AutoHeight()
-								[
-									SNew(SBox)
-									.WidthOverride(650)
+									SNew(SVerticalBox)
+									+ SVerticalBox::Slot()
 									[
-										SNew(SVerticalBox)
-
-										+ SVerticalBox::Slot()
-										.AutoHeight()
-										.HAlign(HAlign_Center)
-										.Padding(FMargin(0, 0, 0, 25.f))
-										[
-											MakeProjectSettingsOptionsBox()
-										]
+										SNew(STextBlock)
+										.TextStyle(FEditorStyle::Get(), "GameProjectDialog.FeatureText")
+										.Text(LOCTEXT("ProjectTemplateClassTypes", "Class Type References:"))
+									]
+									+ SVerticalBox::Slot()
+									.AutoHeight()
+									[
+										SNew(STextBlock)
+										.AutoWrapText(true)
+										.Text(this, &SNewProjectWizard::GetSelectedTemplateClassTypes)
 									]
 								]
+							]
 
+							// Languages available
+							+ SVerticalBox::Slot()
+							.AutoHeight()
+							.Padding(FMargin(0.0f, 5.0f, 0.0f, 5.0f))
+							[
+								SNew(SVerticalBox)
 								+ SVerticalBox::Slot()
-								.AutoHeight()
-								.Padding(FMargin(0, 0, 0, 15.f))
 								[
-									SNew(SRichTextBlock)
-									.Text(LOCTEXT("ProjectPathDescription", "Select a <RichTextBlock.BoldHighlight>location</> for your project to be stored."))
-									.AutoWrapText(true)
-									.DecoratorStyleSet(&FEditorStyle::Get())
-									.ToolTip(IDocumentation::Get()->CreateToolTip(LOCTEXT("ProjectPathDescriptionTooltip", "All of your project content and code will be stored here."), NULL, TEXT("Shared/Editor/NewProjectWizard"), TEXT("ProjectPath")))
+									SNew(STextBlock)
+									.TextStyle(FEditorStyle::Get(), "GameProjectDialog.FeatureText")
+									.Text(LOCTEXT("LanguagesAvailable", "Languages Available:"))
 								]
-
 								+ SVerticalBox::Slot()
 								.AutoHeight()
-								.HAlign(HAlign_Center)
 								[
-									// File path widget
-									SNew(SFilepath)
-									.OnBrowseForFolder(this, &SNewProjectWizard::HandleBrowseButtonClicked)
-									.LabelBackgroundBrush(FEditorStyle::GetBrush("ProjectBrowser.Background"))
-									.LabelBackgroundColor(FLinearColor::White)
-									.FolderPath(this, &SNewProjectWizard::GetCurrentProjectFilePath)
-									.Name(this, &SNewProjectWizard::GetCurrentProjectFileName)
-									.OnFolderChanged(this, &SNewProjectWizard::OnCurrentProjectFilePathChanged)
-									.OnNameChanged(this, &SNewProjectWizard::OnCurrentProjectFileNameChanged)
+									SNew(STextBlock)
+									.AutoWrapText(true)
+									.Text(this, &SNewProjectWizard::GetSelectedTemplateLanguages)
 								]
 							]
 						]
 					]
 				]
 			]
-
-			// Global Error label
-			+SOverlay::Slot()
-			.HAlign(HAlign_Left)
-			.VAlign(VAlign_Bottom)
-			.Padding( UniformPadding / 2 )
-			[
-				SNew(SBorder)
-				.Visibility(this, &SNewProjectWizard::GetGlobalErrorLabelVisibility)
-				.BorderImage(FEditorStyle::GetBrush("GameProjectDialog.ErrorLabelBorder"))
-				.Padding( UniformPadding / 2 )
-				[
-					SNew(SHorizontalBox)
-						
-					+ SHorizontalBox::Slot()
-					.VAlign(VAlign_Center)
-					.Padding(2.f)
-					.AutoWidth()
-					[
-						SNew(SImage)
-						.Image(FEditorStyle::GetBrush("MessageLog.Warning"))
-					]
-
-					+SHorizontalBox::Slot()
-					.VAlign(VAlign_Center)
-					.FillWidth(1.0f)
-					[
-						SNew(STextBlock)
-						.Text(this, &SNewProjectWizard::GetGlobalErrorLabelText)
-						.TextStyle(FEditorStyle::Get(), TEXT("GameProjectDialog.ErrorLabelFont"))
-					]
-
-					// Button/link to the suggested IDE
-					+SHorizontalBox::Slot()
-					.VAlign(VAlign_Center)
-					.HAlign(HAlign_Center)
-					.AutoWidth()
-					.Padding(5.f, 0.f)
-					[
-						SNew(SGetSuggestedIDEWidget)
-					]
-									
-					// A button to close the persistent global error text
-					+SHorizontalBox::Slot()
-					.AutoWidth()
-					.VAlign(VAlign_Center)
-					[
-						SNew(SButton)
-						.ButtonStyle(FEditorStyle::Get(), "NoBorder")
-						.ContentPadding(0.0f) 
-						.OnClicked(this, &SNewProjectWizard::OnCloseGlobalErrorLabelClicked)
-						.Visibility(this, &SNewProjectWizard::GetGlobalErrorLabelCloseButtonVisibility)
-						[
-							SNew(SImage)
-							.Image(FEditorStyle::GetBrush("GameProjectDialog.ErrorLabelCloseButton"))
-						]
-					]
-				]
-			]
-
-			// Project filename error
-			+SOverlay::Slot()
-			.HAlign(HAlign_Left)
-			.VAlign(VAlign_Bottom)
-			.Padding( UniformPadding / 2 )
+		]
+		+ SVerticalBox::Slot()
+		.AutoHeight()
+		[
+			SNew(SBox)
+			.Visibility(this, &SNewProjectWizard::GetTemplateListLocationBoxVisibility)
 			[
-				SNew(SBorder)
-				.BorderImage(FEditorStyle::GetBrush("GameProjectDialog.ErrorLabelBorder"))
-				.Visibility(this, &SNewProjectWizard::GetNameAndLocationErrorLabelVisibility)
-				.Padding(UniformPadding / 2)
-				[
-					SNew(SHorizontalBox)
-					
-					+ SHorizontalBox::Slot()
-					.VAlign(VAlign_Center)
-					.Padding(2.f)
-					.AutoWidth()
-					[
-						SNew(SImage)
-						.Image(FEditorStyle::GetBrush("MessageLog.Warning"))
-					]
-
-					+SHorizontalBox::Slot()
-					.VAlign(VAlign_Center)
-					.AutoWidth()
-					[
-						SNew(STextBlock)
-						.AutoWrapText(true)
-						.Text(this, &SNewProjectWizard::GetNameAndLocationErrorLabelText)
-						.TextStyle(FEditorStyle::Get(), "GameProjectDialog.ErrorLabelFont")
-					]
-				]
+				MakeProjectLocationWidget()
 			]
 		]
 	];
 
-	// Initialize the current page name. Assuming the template page.
-	CurrentPageName = TemplatePageName;
-
-	HandleCategoryChanged(ECheckBoxState::Checked, ActiveCategory);
+	SetCurrentCategory(ActiveCategory);
 
 	UpdateProjectFileValidity();
 }
 END_SLATE_FUNCTION_BUILD_OPTIMIZATION
 
-
-TSharedRef<SWidget> SNewProjectWizard::BuildCategoryTabs()
+bool SNewProjectWizard::ShouldShowProjectSettingsPage() const
 {
-	TSharedRef<SHorizontalBox> TabStrip = SNew(SHorizontalBox);
-
-	TArray<FName> Categories;
-	Templates.GenerateKeyArray(Categories);
-
-	for (const FName& CategoryName : Categories)
-	{
-		TSharedPtr<const FTemplateCategory> Category = FGameProjectGenerationModule::Get().GetCategory(CategoryName);
-
-		TSharedPtr<SHorizontalBox> HorizontalBox;
-
-		TabStrip->AddSlot().AutoWidth()
-		.VAlign(VAlign_Center)
-		.Padding(FMargin(0,0,2.f,0))
-		[
-			SNew(SBox)
-			// Constrain the height to 32px (for the image) plus 5px padding vertically
-			.HeightOverride(32.f + 5.f*2)
-			[
-				SNew( SCheckBox )
-				.Style( FEditorStyle::Get(), "GameProjectDialog.Tab" )
-				.OnCheckStateChanged(this, &SNewProjectWizard::HandleCategoryChanged, CategoryName)
-				.IsChecked(this, &SNewProjectWizard::GetCategoryTabCheckState, CategoryName)
-				.ToolTipText(Category.IsValid() ? Category->Description : FText())
-				.Padding(FMargin(5.f))
-				[
-					SAssignNew(HorizontalBox, SHorizontalBox)
-				]
-			]
-		];
-
-
-		if (Category.IsValid())
-		{
-			HorizontalBox->AddSlot()
-			.AutoWidth()
-			[
-				SNew(SBox)
-				.WidthOverride(32)
-				.HeightOverride(32)
-				[
-					SNew(SImage)
-					.Image(Category->Icon)
-				]
-			];
-		}
-
-		HorizontalBox->AddSlot()
-		.Padding(5.f, 0.f)
-		.VAlign(VAlign_Center)
-		.AutoWidth()
-		[
-			SNew(STextBlock)
-			.TextStyle(FEditorStyle::Get(), "GameProjectDialog.FeatureText")
-			.Text(Category.IsValid() ? Category->Name : FText::FromString(CategoryName.ToString()))
-		];
-	}
-	return TabStrip;
+	bool bSkipSettings = GetSelectedTemplateProperty<bool>(&FTemplateItem::bSkipProjectSettings);
+	return !bSkipSettings;
 }
 
 void SNewProjectWizard::OnSetCopyStarterContent(int32 InCopyStarterContent)
@@ -781,6 +533,12 @@ void SNewProjectWizard::OnSetCopyStarterContent(int32 InCopyStarterContent)
 	bCopyStarterContent = InCopyStarterContent != 0;
 }
 
+EVisibility SNewProjectWizard::GetTemplateListLocationBoxVisibility() const
+{
+	bool bSkipSettings = GetSelectedTemplateProperty<bool>(&FTemplateItem::bSkipProjectSettings);
+	return bSkipSettings ? EVisibility::Visible : EVisibility::Collapsed;
+}
+
 EVisibility SNewProjectWizard::GetStarterContentWarningVisibility() const
 {
 	return (bCopyStarterContent && (SelectedHardwareClassTarget == EHardwareClass::Mobile)) ? EVisibility::Visible : EVisibility::Collapsed;
@@ -809,12 +567,13 @@ void SNewProjectWizard::Tick( const FGeometry& AllottedGeometry, const double In
 	}
 }
 
-
 void SNewProjectWizard::HandleTemplateListViewSelectionChanged(TSharedPtr<FTemplateItem> TemplateItem, ESelectInfo::Type SelectInfo)
 {
 	UpdateProjectFileValidity();
-}
 
+	FString CodeProjectFile = GetSelectedTemplateProperty(&FTemplateItem::CodeProjectFile);
+	bShouldGenerateCode = CodeProjectFile.IsEmpty() == false;
+}
 
 TSharedPtr<FTemplateItem> SNewProjectWizard::GetSelectedTemplateItem() const
 {
@@ -859,46 +618,58 @@ EVisibility SNewProjectWizard::GetSelectedTemplatePreviewVisibility() const
 	return PreviewImage.IsValid() ? EVisibility::Visible : EVisibility::Collapsed;
 }
 
-const FSlateBrush* SNewProjectWizard::GetSelectedTemplateTypeImage() const
+FText SNewProjectWizard::GetSelectedTemplateLanguages() const
 {
 	TSharedPtr<FTemplateItem> SelectedItem = GetSelectedTemplateItem();
 	if (SelectedItem.IsValid())
 	{
-		auto Category = FGameProjectGenerationModule::Get().GetCategory(SelectedItem->Type);
-		if (Category.IsValid())
+		bool bHasBlueprint = !SelectedItem->BlueprintProjectFile.IsEmpty();
+		bool bHasCode = !SelectedItem->CodeProjectFile.IsEmpty();
+
+		if (!bHasBlueprint && !bHasCode)
+		{
+			bHasBlueprint = true;
+			bHasCode = true;
+		}
+
+		if (bHasBlueprint && bHasCode)
+		{
+			return LOCTEXT("ProjectDialog_BlueprintAndCode", "Blueprint, C++");
+		}
+		else if (bHasBlueprint)
+		{
+			return LOCTEXT("ProjectDialog_Blueprint", "Blueprint");
+		}
+		else if (bHasCode)
 		{
-			return Category->Image;
+			return LOCTEXT("ProjectDialog_Code", "C++");
 		}
 	}
-	return nullptr;
-}
 
+	return FText::GetEmpty();
+}
 
 FText SNewProjectWizard::GetCurrentProjectFileName() const
 {
 	return FText::FromString( CurrentProjectFileName );
 }
 
-
 FString SNewProjectWizard::GetCurrentProjectFileNameStringWithExtension() const
 {
 	return CurrentProjectFileName + TEXT(".") + FProjectDescriptor::GetExtension();
 }
 
-
 void SNewProjectWizard::OnCurrentProjectFileNameChanged(const FText& InValue)
 {
 	CurrentProjectFileName = InValue.ToString();
 	UpdateProjectFileValidity();
 }
 
-
 FText SNewProjectWizard::GetCurrentProjectFilePath() const
 {
 	return FText::FromString(CurrentProjectFilePath);
 }
 
-
 FString SNewProjectWizard::GetCurrentProjectFileParentFolder() const
 {
 	if ( CurrentProjectFilePath.EndsWith(TEXT("/")) || CurrentProjectFilePath.EndsWith("\\") )
@@ -911,7 +682,6 @@ FString SNewProjectWizard::GetCurrentProjectFileParentFolder() const
 	}
 }
 
-
 void SNewProjectWizard::OnCurrentProjectFilePathChanged(const FText& InValue)
 {
 	CurrentProjectFilePath = InValue.ToString();
@@ -919,13 +689,11 @@ void SNewProjectWizard::OnCurrentProjectFilePathChanged(const FText& InValue)
 	UpdateProjectFileValidity();
 }
 
-
 FString SNewProjectWizard::GetProjectFilenameWithPathLabelText() const
 {
 	return GetProjectFilenameWithPath();
 }
 
-
 FString SNewProjectWizard::GetProjectFilenameWithPath() const
 {
 	if ( CurrentProjectFilePath.IsEmpty() )
@@ -975,45 +743,17 @@ FReply SNewProjectWizard::HandleBrowseButtonClicked()
 	return FReply::Handled();
 }
 
+
 void SNewProjectWizard::HandleTemplateListViewDoubleClick( TSharedPtr<FTemplateItem> TemplateItem )
 {
-	// Advance to the name/location page
-	const int32 NamePageIdx = 1;
-	if ( MainWizard->CanShowPage(NamePageIdx) )
-	{
-		MainWizard->ShowPage(NamePageIdx);
-	}
+	OnTemplateDoubleClick.ExecuteIfBound();
 }
 
-
-bool SNewProjectWizard::IsCreateProjectEnabled() const
+bool SNewProjectWizard::CanCreateProject() const
 {
-	if ( CurrentPageName == NAME_None )//|| CurrentPageName == TemplatePageName )
-	{
-		return false;
-	}
-
 	return bLastGlobalValidityCheckSuccessful && bLastNameAndLocationValidityCheckSuccessful;
 }
 
-
-bool SNewProjectWizard::HandlePageCanShow( FName PageName ) const
-{
-	if (PageName == NameAndLocationPageName)
-	{
-		return bLastGlobalValidityCheckSuccessful;
-	}
-
-	return true;
-}
-
-
-void SNewProjectWizard::OnPageVisited(FName NewPageName)
-{
-	CurrentPageName = NewPageName;
-}
-
-
 EVisibility SNewProjectWizard::GetGlobalErrorLabelVisibility() const
 {
 	const bool bIsVisible = GetNameAndLocationErrorLabelText().IsEmpty() && !GetGlobalErrorLabelText().IsEmpty();
@@ -1068,56 +808,14 @@ FText SNewProjectWizard::GetNameAndLocationErrorLabelText() const
 
 TMap<FName, TArray<TSharedPtr<FTemplateItem>> >& SNewProjectWizard::FindTemplateProjects()
 {
-	// Default to showing the blueprint category
-	ActiveCategory = FApp::IsEnterpriseInstalled() ? FTemplateCategory::EnterpriseCategoryName : FTemplateCategory::BlueprintCategoryName;
-	
 	// Clear the list out first - or we could end up with duplicates
 	Templates.Empty();
 
-	// Add some default non-data driven templates
-	Templates.FindOrAdd(FTemplateCategory::BlueprintCategoryName).Add(MakeShareable(new FTemplateItem(
-		LOCTEXT("BlankProjectName", "Blank"),
-		LOCTEXT("BlankProjectDescription", "A clean empty project with no code."),
-		false, FTemplateCategory::BlueprintCategoryName,
-		TEXT("_1"),			// SortKey
-		TEXT(""),			// No filename, this is a generation template
-		MakeShareable( new FSlateBrush( *FEditorStyle::GetBrush("GameProjectDialog.BlankProjectThumbnail") ) ),
-		MakeShareable( new FSlateBrush( *FEditorStyle::GetBrush("GameProjectDialog.BlankProjectPreview") ) ),
-		TEXT(""),		// No class types
-		TEXT("")		// No asset types
-		)) );
-
-	Templates.FindOrAdd(FTemplateCategory::CodeCategoryName).Add(MakeShareable(new FTemplateItem(
-		LOCTEXT("BasicCodeProjectName", "Basic Code"),
-		LOCTEXT("BasicCodeProjectDescription", "An empty project with some basic game framework code classes created."),
-		true, FTemplateCategory::CodeCategoryName,
-		TEXT("_2"),			// SortKey
-		TEXT(""),			// No filename, this is a generation template
-		MakeShareable( new FSlateBrush( *FEditorStyle::GetBrush("GameProjectDialog.BasicCodeThumbnail") ) ),
-		MakeShareable( new FSlateBrush( *FEditorStyle::GetBrush("GameProjectDialog.BlankProjectPreview") ) ),
-		TEXT(""),		// No class types
-		TEXT("")		// No asset types
-		)) );
-
-	if (FPaths::DirectoryExists(FPaths::EnterpriseDir()))
-	{
-		Templates.FindOrAdd(FTemplateCategory::EnterpriseCategoryName).Add(MakeShareable(new FTemplateItem(
-			LOCTEXT("BlankEnterpriseProjectName", "Blank"),
-			LOCTEXT("BlankEnterpriseProjectDescription", "A clean empty Unreal Studio project with no code."),
-			false, FTemplateCategory::EnterpriseCategoryName,
-			TEXT("_1"),			// SortKey
-			TEXT(""),			// No filename, this is a generation template
-			MakeShareable(new FSlateBrush(*FEditorStyle::GetBrush("GameProjectDialog.BlankProjectThumbnail"))),
-			MakeShareable(new FSlateBrush(*FEditorStyle::GetBrush("GameProjectDialog.BlankProjectPreview"))),
-			TEXT(""),		// No class types
-			TEXT("")		// No asset types
-		)));
-	}
 	// Now discover and all data driven templates
 	TArray<FString> TemplateRootFolders;
 
 	// @todo rocket make template folder locations extensible.
-	TemplateRootFolders.Add( FPaths::RootDir() + TEXT("Templates") );
+	TemplateRootFolders.Add(FPaths::RootDir() + TEXT("Templates"));
 
 	// Add the Enterprise templates
 	TemplateRootFolders.Add(FPaths::EnterpriseDir() + TEXT("Templates"));
@@ -1143,6 +841,7 @@ TMap<FName, TArray<TSharedPtr<FTemplateItem>> >& SNewProjectWizard::FindTemplate
 	for ( auto TemplateRootFolderIt = TemplateRootFolders.CreateConstIterator(); TemplateRootFolderIt; ++TemplateRootFolderIt )
 	{
 		const FString Root = *TemplateRootFolderIt;
+		
 		const FString SearchString = Root / TEXT("*");
 		TArray<FString> TemplateFolders;
 		IFileManager::Get().FindFiles(TemplateFolders, *SearchString, /*Files=*/false, /*Directories=*/true);
@@ -1152,6 +851,8 @@ TMap<FName, TArray<TSharedPtr<FTemplateItem>> >& SNewProjectWizard::FindTemplate
 		}
 	}
 
+	TArray<TSharedPtr<FTemplateItem>> FoundTemplates;
+
 	// Add a template item for every discovered project
 	for ( auto TemplateFolderIt = AllTemplateFolders.CreateConstIterator(); TemplateFolderIt; ++TemplateFolderIt )
 	{
@@ -1165,33 +866,74 @@ TMap<FName, TArray<TSharedPtr<FTemplateItem>> >& SNewProjectWizard::FindTemplate
 				// Make sure a TemplateDefs ini file exists
 				const FString Root = *TemplateFolderIt;
 				UTemplateProjectDefs* TemplateDefs = GameProjectUtils::LoadTemplateDefs(Root);
-				if ( TemplateDefs )
+				if (TemplateDefs)
 				{
-					// Ignore any templates whoose definition says we cannot use to create a project
+					// Ignore any templates whose definition says we cannot use to create a project
 					if( TemplateDefs->bAllowProjectCreation == false )
 						continue;
-					// Found a template. Add it to the template items list.
-					FString ProjectFilename = Root / FoundProjectFiles[0];
-					FText TemplateName = TemplateDefs->GetDisplayNameText();
-					FText TemplateDescription = TemplateDefs->GetLocalizedDescription();
-					FString ClassTypes = TemplateDefs->ClassTypes;
-					FString AssetTypes = TemplateDefs->AssetTypes;
-					
+
+					FString ProjectFile = Root / FoundProjectFiles[0];
+
 					// If no template name was specified for the current culture, just use the project name
-					if ( TemplateName.IsEmpty() )
+					TArray<FName> TemplateCategories = TemplateDefs->Categories;
+					if (TemplateCategories.Num() == 0)
 					{
-						TemplateName = FText::FromString(FPaths::GetBaseFilename(ProjectFilename));
+						TemplateCategories.Add("Game");
+					}
+
+					FString TemplateKey = Root;
+					TemplateKey.RemoveFromEnd("BP");
+
+					TSharedPtr<FTemplateItem>* ExistingTemplate = FoundTemplates.FindByPredicate([&TemplateKey](TSharedPtr<FTemplateItem> Item)
+					{
+						return Item->Key == TemplateKey;
+					});
+
+					if (ExistingTemplate != nullptr)
+					{
+						if (TemplateDefs->GeneratesCode(Root))
+						{
+							(*ExistingTemplate)->CodeProjectFile = ProjectFile;
+						}
+						else
+						{
+							(*ExistingTemplate)->BlueprintProjectFile = ProjectFile;
+						}
+
+						continue;
+					}
+					
+					// Did not find an existing template. Create a new one to add to the template list.
+					TSharedPtr<FTemplateItem> Template = MakeShareable(new FTemplateItem());
+					Template->Key = TemplateKey;
+					Template->Categories = TemplateCategories;
+					Template->Description = TemplateDefs->GetLocalizedDescription();
+					Template->ClassTypes = TemplateDefs->ClassTypes;
+					Template->AssetTypes = TemplateDefs->AssetTypes;
+					Template->bSkipProjectSettings = TemplateDefs->bSkipProjectSettings;
+					Template->HiddenSettings = TemplateDefs->HiddenSettings;
+
+					Template->Name = TemplateDefs->GetDisplayNameText();
+					if (Template->Name.IsEmpty())
+					{
+						Template->Name = FText::FromString(FPaths::GetBaseFilename(ProjectFile));
 					}
 
 					// Only generate code if the template has a source folder
-					const bool bGenerateCode = TemplateDefs->GeneratesCode(Root);
+					if (TemplateDefs->GeneratesCode(Root))
+					{
+						Template->CodeProjectFile = ProjectFile;
+					}
+					else
+					{
+						Template->BlueprintProjectFile = ProjectFile;
+					}
 
-					TSharedPtr<FSlateDynamicImageBrush> ThumbnailBrush;
 					const FString ThumbnailPNGFile = (Root + TEXT("/Media/") + FoundProjectFiles[0]).Replace(TEXT(".uproject"), TEXT(".png"));
 					if ( FPlatformFileManager::Get().GetPlatformFile().FileExists(*ThumbnailPNGFile) )
 					{
 						const FName BrushName = FName(*ThumbnailPNGFile);
-						ThumbnailBrush = MakeShareable( new FSlateDynamicImageBrush(BrushName , FVector2D(128,128) ) );
+						Template->Thumbnail = MakeShareable( new FSlateDynamicImageBrush(BrushName, FVector2D(128,128) ) );
 					}
 
 					TSharedPtr<FSlateDynamicImageBrush> PreviewBrush;
@@ -1199,39 +941,24 @@ TMap<FName, TArray<TSharedPtr<FTemplateItem>> >& SNewProjectWizard::FindTemplate
 					if ( FPlatformFileManager::Get().GetPlatformFile().FileExists(*PreviewPNGFile) )
 					{
 						const FName BrushName = FName(*PreviewPNGFile);
-						PreviewBrush = MakeShareable( new FSlateDynamicImageBrush(BrushName , FVector2D(512,256) ) );
+						Template->PreviewImage = MakeShareable( new FSlateDynamicImageBrush(BrushName, FVector2D(512,256) ) );
 					}
 
 					// Get the sort key
 					FString SortKey = TemplateDefs->SortKey;
+					const FString CleanFilename = FPaths::GetCleanFilename(ProjectFile);
 					if(SortKey.Len() == 0)
 					{
-						SortKey = FPaths::GetCleanFilename(ProjectFilename);
-					}
-					if (FPaths::GetCleanFilename(ProjectFilename) == GameProjectUtils::GetDefaultProjectTemplateFilename())
-					{
-						SortKey = TEXT("_0");
+						SortKey = CleanFilename;
 					}
 
-					// Assign the template to the correct category. If the template has no explicit category assigned, we assign it to either C++ or blueprint
-					FName Category = TemplateDefs->Category;
-					if (Category.IsNone())
+					if (CleanFilename == GameProjectUtils::GetDefaultProjectTemplateFilename())
 					{
-						Category = bGenerateCode ? FTemplateCategory::CodeCategoryName : FTemplateCategory::BlueprintCategoryName;
+						SortKey = TEXT("_0");
 					}
+					Template->SortKey = SortKey;
 
-					Templates.FindOrAdd(Category).Add(MakeShareable(new FTemplateItem(
-						TemplateName,
-						TemplateDescription,
-						bGenerateCode,
-						Category,
-						MoveTemp(SortKey),
-						MoveTemp(ProjectFilename),
-						ThumbnailBrush,
-						PreviewBrush,
-						ClassTypes,
-						AssetTypes
-					)));
+					FoundTemplates.Add(Template);
 				}
 			}
 			else
@@ -1241,6 +968,15 @@ TMap<FName, TArray<TSharedPtr<FTemplateItem>> >& SNewProjectWizard::FindTemplate
 			}
 		}
 	}
+
+	for (const TSharedPtr<FTemplateItem>& Template : FoundTemplates)
+	{
+		for (const FName& Category : Template->Categories)
+		{
+			Templates.FindOrAdd(Category).Add(Template);
+		}
+	}
+
 	return Templates;
 }
 
@@ -1381,13 +1117,16 @@ void SNewProjectWizard::UpdateProjectFileValidity( )
 	bPreventPeriodicValidityChecksUntilNextChange = false;
 }
 
-
-bool SNewProjectWizard::IsCompilerRequired( ) const
+bool SNewProjectWizard::IsCompilerRequired() const
 {
 	TSharedPtr<FTemplateItem> SelectedTemplate = GetSelectedTemplateItem();
-	return SelectedTemplate.IsValid() && SelectedTemplate->bGenerateCode;
-}
 
+	if (SelectedTemplate.IsValid())
+	{
+		return bShouldGenerateCode && !SelectedTemplate->CodeProjectFile.IsEmpty();
+	}
+	return false;
+}
 
 bool SNewProjectWizard::CreateProject( const FString& ProjectFile )
 {
@@ -1402,11 +1141,19 @@ bool SNewProjectWizard::CreateProject( const FString& ProjectFile )
 
 	FText FailReason, FailLog;
 
-	FProjectInformation ProjectInfo(ProjectFile, SelectedTemplate->bGenerateCode, bCopyStarterContent, SelectedTemplate->ProjectFile);
+	FProjectInformation ProjectInfo;
+	ProjectInfo.ProjectFilename = ProjectFile;
+	ProjectInfo.bShouldGenerateCode = bShouldGenerateCode;
+	ProjectInfo.bCopyStarterContent = bCopyStarterContent;
+	ProjectInfo.TemplateFile = bShouldGenerateCode ? SelectedTemplate->CodeProjectFile : SelectedTemplate->BlueprintProjectFile;
+	ProjectInfo.TemplateCategory = ActiveCategory;
 	ProjectInfo.TargetedHardware = SelectedHardwareClassTarget;
 	ProjectInfo.DefaultGraphicsPerformance = SelectedGraphicsPreset;
-	ProjectInfo.bIsEnterpriseProject = (SelectedTemplate->Type == FTemplateCategory::EnterpriseCategoryName) || FApp::IsEnterpriseInstalled();
-	ProjectInfo.bForceExtendedLuminanceRange = SelectedTemplate->ProjectFile.IsEmpty();
+	ProjectInfo.bForceExtendedLuminanceRange = ProjectInfo.TemplateFile.IsEmpty();
+	ProjectInfo.bEnableVR = bEnableVR;
+
+	// @todo (sebastiann): This is a poor way to determine Enterprise-ness.
+	ProjectInfo.bIsEnterpriseProject = SelectedTemplate->Category != TEXT("Game");
 
 	if (!GameProjectUtils::CreateProject(ProjectInfo, FailReason, FailLog))
 	{
@@ -1435,7 +1182,7 @@ bool SNewProjectWizard::CreateProject( const FString& ProjectFile )
 
 void SNewProjectWizard::CreateAndOpenProject( )
 {
-	if( !IsCreateProjectEnabled() )
+	if( !CanCreateProject() )
 	{
 		return;
 	}
@@ -1449,7 +1196,7 @@ void SNewProjectWizard::CreateAndOpenProject( )
 	// Prevent periodic validity checks. This is to prevent a brief error message about the project already existing while you are exiting.
 	bPreventPeriodicValidityChecksUntilNextChange = true;
 
-	if( GetSelectedTemplateItem()->bGenerateCode )
+	if (bShouldGenerateCode)
 	{
 	    // If the engine is installed it is already compiled, so we can try to build and open a new project immediately. Non-installed situations might require building
 	    // the engine (especially the case when binaries came from P4), so we only open the IDE for that.
@@ -1542,30 +1289,8 @@ void SNewProjectWizard::DisplayError( const FText& ErrorText )
 /* SNewProjectWizard event handlers
  *****************************************************************************/
 
-bool SNewProjectWizard::HandleCreateProjectWizardCanFinish( ) const
-{
-	return IsCreateProjectEnabled();
-}
-
-
-void SNewProjectWizard::HandleCreateProjectWizardFinished( )
-{
-	CreateAndOpenProject();
-}
-
-ECheckBoxState SNewProjectWizard::GetCategoryTabCheckState(FName Category) const
-{
-	return Category == ActiveCategory ? ECheckBoxState::Checked : ECheckBoxState::Unchecked;
-}
-
-void SNewProjectWizard::HandleCategoryChanged(ECheckBoxState CheckState, FName Category)
+void SNewProjectWizard::SetCurrentCategory(FName Category)
 {
-	if (CheckState != ECheckBoxState::Checked)
-	{
-		return;
-	}
-
-	ActiveCategory = Category;
 	FilteredTemplateList = Templates.FindRef(Category);
 
 	// Sort the template folders
@@ -1579,21 +1304,182 @@ void SNewProjectWizard::HandleCategoryChanged(ECheckBoxState CheckState, FName C
 	}
 	TemplateListView->RequestListRefresh();
 
-	if (ActiveCategory == FTemplateCategory::EnterpriseCategoryName)
-	{
-		bProjectSettingsHidden = true;
-		bCopyStarterContent = false;
-		SetGraphicsPreset(EGraphicsPreset::Maximum);
-		SetHardwareClassTarget(EHardwareClass::Desktop);		
-		ProjectSettingsDescriptionBox->SetVisibility(EVisibility::Collapsed);
-		ProjectSettingsOptionsBox->SetVisibility(EVisibility::Collapsed);
-	}
-	else if (bProjectSettingsHidden)
-	{
-		bProjectSettingsHidden = false;
-		ProjectSettingsDescriptionBox->SetVisibility(EVisibility::Visible);
-		ProjectSettingsOptionsBox->SetVisibility(EVisibility::Visible);
-	}
+	ActiveCategory = Category;
+}
+
+TSharedRef<SWidget> SNewProjectWizard::MakeProjectLocationWidget()
+{
+	return SNew(SVerticalBox)
+	+ SVerticalBox::Slot()
+	.HAlign(HAlign_Center)
+	.Padding(8)
+	[
+		SNew(SVerticalBox)
+		+ SVerticalBox::Slot()
+		.AutoHeight()
+		.Padding(0, 0, 0, 8)
+		[
+			SNew(SRichTextBlock)
+			.Text(LOCTEXT("ProjectPathDescription", "Select a <RichTextBlock.BoldHighlight>location</> for your project to be stored."))
+			.AutoWrapText(true)
+			.DecoratorStyleSet(&FEditorStyle::Get())
+			.ToolTip(IDocumentation::Get()->CreateToolTip(LOCTEXT("ProjectPathDescriptionTooltip", "All of your project content and code will be stored here."), NULL, TEXT("Shared/Editor/NewProjectWizard"), TEXT("ProjectPath")))
+		]
+
+		+ SVerticalBox::Slot()
+		.AutoHeight()
+		[
+			// File path widget
+			SNew(SFilepath)
+			.OnBrowseForFolder(this, &SNewProjectWizard::HandleBrowseButtonClicked)
+			.LabelBackgroundBrush(FEditorStyle::GetBrush("ProjectBrowser.Background"))
+			.LabelBackgroundColor(FLinearColor::White)
+			.FolderPath(this, &SNewProjectWizard::GetCurrentProjectFilePath)
+			.Name(this, &SNewProjectWizard::GetCurrentProjectFileName)
+			.OnFolderChanged(this, &SNewProjectWizard::OnCurrentProjectFilePathChanged)
+			.OnNameChanged(this, &SNewProjectWizard::OnCurrentProjectFileNameChanged)
+		]
+	];
+}
+
+TSharedRef<SWidget> SNewProjectWizard::CreateProjectSettingsPage()
+{
+	const float UniformPadding = 16.f;
+
+	ProjectSettingsDescriptionBox = MakeProjectSettingsDescriptionBox();
+	ProjectSettingsOptionsBox = MakeProjectSettingsOptionsBox();
+
+	TSharedRef<SWidget> PageWidget = SNew(SOverlay)
+	+ SOverlay::Slot()
+	[
+		SNew(SVerticalBox)
+		+ SVerticalBox::Slot()
+		.AutoHeight()
+		.Padding(0, 8, 0, 8)
+		[
+			ProjectSettingsDescriptionBox.ToSharedRef()
+		]
+		+ SVerticalBox::Slot()
+		.Padding(0)
+		.VAlign(VAlign_Fill)
+		[
+			SNew(SBorder)
+			.Padding(UniformPadding)
+			.BorderImage(FEditorStyle::GetBrush("ToolPanel.GroupBorder"))
+			[
+				SNew(SScrollBox)
+				+ SScrollBox::Slot()
+				.Padding(0)
+				[
+					SNew(SVerticalBox)
+					+ SVerticalBox::Slot()
+					.Padding(FMargin(0, 0, 0, UniformPadding))
+					[
+						ProjectSettingsOptionsBox.ToSharedRef()
+					]
+				]
+			]
+		]
+		+ SVerticalBox::Slot()
+		.Padding(0)
+		.AutoHeight()
+		[
+			MakeProjectLocationWidget()
+		]
+	]
+	// Global Error label
+	+ SOverlay::Slot()
+	.HAlign(HAlign_Left)
+	.VAlign(VAlign_Bottom)
+	.Padding( 8 )
+	[
+		SNew(SBorder)
+		.Visibility(this, &SNewProjectWizard::GetGlobalErrorLabelVisibility)
+		.BorderImage(FEditorStyle::GetBrush("GameProjectDialog.ErrorLabelBorder"))
+		.Padding( 8 )
+		[
+			SNew(SHorizontalBox)
+						
+			+ SHorizontalBox::Slot()
+			.VAlign(VAlign_Center)
+			.Padding(2.f)
+			.AutoWidth()
+			[
+				SNew(SImage)
+				.Image(FEditorStyle::GetBrush("MessageLog.Warning"))
+			]
+
+			+SHorizontalBox::Slot()
+			.VAlign(VAlign_Center)
+			.FillWidth(1.0f)
+			[
+				SNew(STextBlock)
+				.Text(this, &SNewProjectWizard::GetGlobalErrorLabelText)
+				.TextStyle(FEditorStyle::Get(), TEXT("GameProjectDialog.ErrorLabelFont"))
+			]
+
+			// Button/link to the suggested IDE
+			+SHorizontalBox::Slot()
+			.VAlign(VAlign_Center)
+			.HAlign(HAlign_Center)
+			.AutoWidth()
+			.Padding(5.f, 0.f)
+			[
+				SNew(SGetSuggestedIDEWidget)
+			]
+									
+			// A button to close the persistent global error text
+			+SHorizontalBox::Slot()
+			.AutoWidth()
+			.VAlign(VAlign_Center)
+			[
+				SNew(SButton)
+				.ButtonStyle(FEditorStyle::Get(), "NoBorder")
+				.ContentPadding(0.0f) 
+				.OnClicked(this, &SNewProjectWizard::OnCloseGlobalErrorLabelClicked)
+				.Visibility(this, &SNewProjectWizard::GetGlobalErrorLabelCloseButtonVisibility)
+				[
+					SNew(SImage)
+					.Image(FEditorStyle::GetBrush("GameProjectDialog.ErrorLabelCloseButton"))
+				]
+			]
+		]
+	]
+	// Project filename error
+	+SOverlay::Slot()
+	.HAlign(HAlign_Left)
+	.VAlign(VAlign_Bottom)
+	.Padding( 8 )
+	[
+		SNew(SBorder)
+		.BorderImage(FEditorStyle::GetBrush("GameProjectDialog.ErrorLabelBorder"))
+		.Visibility(this, &SNewProjectWizard::GetNameAndLocationErrorLabelVisibility)
+		.Padding(8)
+		[
+			SNew(SHorizontalBox)
+					
+			+ SHorizontalBox::Slot()
+			.VAlign(VAlign_Center)
+			.Padding(2.f)
+			.AutoWidth()
+			[
+				SNew(SImage)
+				.Image(FEditorStyle::GetBrush("MessageLog.Warning"))
+			]
+
+			+SHorizontalBox::Slot()
+			.VAlign(VAlign_Center)
+			.AutoWidth()
+			[
+				SNew(STextBlock)
+				.AutoWrapText(true)
+				.Text(this, &SNewProjectWizard::GetNameAndLocationErrorLabelText)
+				.TextStyle(FEditorStyle::Get(), "GameProjectDialog.ErrorLabelFont")
+			]
+		]
+	];
+
+	return PageWidget;
 }
 
 void SNewProjectWizard::SetHardwareClassTarget(EHardwareClass::Type InHardwareClass)
@@ -1606,87 +1492,191 @@ void SNewProjectWizard::SetGraphicsPreset(EGraphicsPreset::Type InGraphicsPreset
 	SelectedGraphicsPreset = InGraphicsPreset;
 }
 
-TSharedRef<SRichTextBlock> SNewProjectWizard::MakeProjectSettingsDescriptionBox()
+int32 SNewProjectWizard::OnGetBlueprintOrCppIndex() const
+{
+	return bShouldGenerateCode ? 1 : 0;
+}
+
+void SNewProjectWizard::OnSetBlueprintOrCppIndex(int32 Index)
+{
+	bShouldGenerateCode = Index == 1;
+}
+
+TSharedRef<SRichTextBlock> SNewProjectWizard::MakeProjectSettingsDescriptionBox() const
 {
-	TSharedPtr<SRichTextBlock> Widget = SNew(SRichTextBlock)
-		.Text(LOCTEXT("ProjectSettingsDescription", "Choose some <RichTextBlock.BoldHighlight>settings</> for your project.  Don't worry, you can change these later in the <RichTextBlock.BoldHighlight>Target Hardware</> section of <RichTextBlock.BoldHighlight>Project Settings</>.  You can also add the <RichTextBlock.BoldHighlight>Starter Content</> to your project later using <RichTextBlock.BoldHighlight>Content Browser</>."))
+	TSharedRef<SRichTextBlock> Widget = SNew(SRichTextBlock)
+		.Text(LOCTEXT("ProjectSettingsDescription", "Choose some <RichTextBlock.BoldHighlight>settings</> for your project."))
 		.AutoWrapText(true)
 		.DecoratorStyleSet(&FEditorStyle::Get())
 		.ToolTip(IDocumentation::Get()->CreateToolTip(LOCTEXT("HardwareTargetTooltip", "These settings will choose good defaults for a number of other settings in the project such as post-processing flags and touch input emulation using the mouse."), NULL, TEXT("Shared/Editor/NewProjectWizard"), TEXT("TargetHardware")));
 
-	ProjectSettingsDescriptionBox = Widget;
+	return Widget;
+}
+
+static void AddToProjectSettingsGrid(TSharedRef<SGridPanel> Grid, const TSharedRef<SWidget>& Enum, const TSharedRef<SWidget>& Description, FIntPoint& Slot)
+{
+	Grid->AddSlot(Slot.X, Slot.Y)
+	.Padding(8)
+	[
+		Enum
+	];
+
+	Grid->AddSlot(Slot.X + 1, Slot.Y)
+	.VAlign(VAlign_Center)
+	.Padding(8)
+	[
+		Description
+	];
+
+	Slot.X += 2;
 
-	return ProjectSettingsDescriptionBox.ToSharedRef();
+	if (Slot.X > 2)
+	{
+		Slot.X = 0;
+		Slot.Y++;
+	}
 }
 
-TSharedRef<SHorizontalBox> SNewProjectWizard::MakeProjectSettingsOptionsBox()
+TSharedRef<SWidget> SNewProjectWizard::MakeProjectSettingsOptionsBox() const
 {
+	static const int EnumWidth = 160;
+
 	IHardwareTargetingModule& HardwareTargeting = IHardwareTargetingModule::Get();
 
-	TSharedPtr<SWidget> StartContentCombo;
+	FIntPoint CurrentSlot(0, 0);
+
+	TSharedRef<SGridPanel> GridPanel = SNew(SGridPanel)
+		.FillColumn(1, 1.0f)
+		.FillColumn(3, 1.0f);
+
+	TArray<ETemplateSetting> HiddenSettings = GetSelectedTemplateProperty(&FTemplateItem::HiddenSettings);
+
+	if (!HiddenSettings.Contains(ETemplateSetting::Languages))
 	{
-		TArray<SDecoratedEnumCombo<int32>::FComboOption> StarterContentInfo;
-		StarterContentInfo.Add(SDecoratedEnumCombo<int32>::FComboOption(
-			0, FSlateIcon(FEditorStyle::GetStyleSetName(), "GameProjectDialog.NoStarterContent"), LOCTEXT("NoStarterContent", "No Starter Content")));
+		bool bIsCodeAvailable = !GetSelectedTemplateProperty(&FTemplateItem::CodeProjectFile).IsEmpty();
+		bool bIsBlueprintAvailable = !GetSelectedTemplateProperty(&FTemplateItem::BlueprintProjectFile).IsEmpty();
 
-		// Only add the option to add starter content if its there to add !
-		bool bIsStarterAvailable = GameProjectUtils::IsStarterContentAvailableForNewProjects();
-		StarterContentInfo.Add(SDecoratedEnumCombo<int32>::FComboOption(
-			1, FSlateIcon(FEditorStyle::GetStyleSetName(), "GameProjectDialog.IncludeStarterContent"), LOCTEXT("IncludeStarterContent", "With Starter Content"), bIsStarterAvailable));
-		StartContentCombo = SNew(SDecoratedEnumCombo<int32>, MoveTemp(StarterContentInfo))
-			.SelectedEnum(this, &SNewProjectWizard::GetCopyStarterContentIndex)
-			.OnEnumChanged(this, &SNewProjectWizard::OnSetCopyStarterContent)
-			.ToolTipText(LOCTEXT("CopyStarterContent_ToolTip", "Enable to include an additional content pack containing simple placeable meshes with basic materials and textures.\nYou can opt out of including this to create a project that only has the bare essentials for the selected project template."));
+		// if neither is available, then this is a blank template, so both are available
+		if (bIsCodeAvailable && bIsBlueprintAvailable)
+		{
+			bIsCodeAvailable = true;
+			bIsBlueprintAvailable = true;
+		}
+
+		TArray<SDecoratedEnumCombo<int32>::FComboOption> BlueprintOrCppOptions;
+		BlueprintOrCppOptions.Add(SDecoratedEnumCombo<int32>::FComboOption(
+			0, FSlateIcon(FEditorStyle::GetStyleSetName(), "GameProjectDialog.BlueprintImage_64"), LOCTEXT("ProjectDialog_Blueprint", "Blueprint"), bIsBlueprintAvailable));
+
+		BlueprintOrCppOptions.Add(SDecoratedEnumCombo<int32>::FComboOption(
+			1, FSlateIcon(FEditorStyle::GetStyleSetName(), "GameProjectDialog.CodeImage_64"), LOCTEXT("ProjectDialog_Code", "C++"), bIsCodeAvailable));
+
+		TSharedRef<SDecoratedEnumCombo<int32>> Enum = SNew(SDecoratedEnumCombo<int32>, MoveTemp(BlueprintOrCppOptions))
+			.SelectedEnum(this, &SNewProjectWizard::OnGetBlueprintOrCppIndex)
+			.OnEnumChanged(this, &SNewProjectWizard::OnSetBlueprintOrCppIndex)
+			.Orientation(Orient_Vertical);
+
+		TSharedRef<SRichTextBlock> Description = SNew(SRichTextBlock)
+			.Text(LOCTEXT("ProjectDialog_BlueprintOrCppDescription", "Choose whether to create a Blueprint or C++ project."))
+			.AutoWrapText(true)
+			.DecoratorStyleSet(&FEditorStyle::Get());
+
+		AddToProjectSettingsGrid(GridPanel, Enum, Description, CurrentSlot);
 	}
 
-	TSharedPtr<SHorizontalBox> Widget = SNew(SHorizontalBox)
+	if (!HiddenSettings.Contains(ETemplateSetting::HardwareTarget))
+	{
+		TSharedRef<SWidget> Enum = HardwareTargeting.MakeHardwareClassTargetCombo(
+			FOnHardwareClassChanged::CreateSP(this, &SNewProjectWizard::SetHardwareClassTarget),
+			TAttribute<EHardwareClass::Type>(this, &SNewProjectWizard::GetHardwareClassTarget),
+			Orient_Vertical);
+		
+		TSharedRef<SRichTextBlock> Description = SNew(SRichTextBlock)
+			.Text(LOCTEXT("HardwareClassTargetDescription", "Choose the closest equivalent target platform. Don't worry, you can change this later in the <RichTextBlock.BoldHighlight>Target Hardware</> section of <RichTextBlock.BoldHighlight>Project Settings</>."))
+			.AutoWrapText(true)
+			.DecoratorStyleSet(&FEditorStyle::Get());
 
-		+ SHorizontalBox::Slot()
-		.AutoWidth()
-		[
-			HardwareTargeting.MakeHardwareClassTargetCombo(
-				FOnHardwareClassChanged::CreateSP(this, &SNewProjectWizard::SetHardwareClassTarget),
-				TAttribute<EHardwareClass::Type>(this, &SNewProjectWizard::GetHardwareClassTarget),
-				Orient_Horizontal
-			)
-		]
+		AddToProjectSettingsGrid(GridPanel, Enum, Description, CurrentSlot);
+	}
 
-		+ SHorizontalBox::Slot()
-		.AutoWidth()
-		.Padding(FMargin(30, 0))
-		[
-			HardwareTargeting.MakeGraphicsPresetTargetCombo(
-				FOnGraphicsPresetChanged::CreateSP(this, &SNewProjectWizard::SetGraphicsPreset),
-				TAttribute<EGraphicsPreset::Type>(this, &SNewProjectWizard::GetGraphicsPreset),
-				Orient_Horizontal
-			)
-		]
 
-		+ SHorizontalBox::Slot()
-		.AutoWidth()
-		[
-			SNew(SOverlay)
-			+ SOverlay::Slot()
-			[
-				StartContentCombo.ToSharedRef()
-			]
+	if (!HiddenSettings.Contains(ETemplateSetting::GraphicsPreset))
+	{
+		TSharedRef<SWidget> Enum = HardwareTargeting.MakeGraphicsPresetTargetCombo(
+			FOnGraphicsPresetChanged::CreateSP(this, &SNewProjectWizard::SetGraphicsPreset),
+			TAttribute<EGraphicsPreset::Type>(this, &SNewProjectWizard::GetGraphicsPreset),
+			Orient_Vertical);
+
+		TSharedRef<SRichTextBlock> Description = SNew(SRichTextBlock)
+			.Text(LOCTEXT("GraphicsPresetDescription", "Choose the performance characteristics of your project."))
+			.AutoWrapText(true)
+			.DecoratorStyleSet(&FEditorStyle::Get());
+
+		AddToProjectSettingsGrid(GridPanel, Enum, Description, CurrentSlot);
+	}
+
+	if (!HiddenSettings.Contains(ETemplateSetting::StarterContent))
+	{
+		TArray<SDecoratedEnumCombo<int32>::FComboOption> StarterContentOptions;
+		StarterContentOptions.Add(SDecoratedEnumCombo<int32>::FComboOption(
+			0, FSlateIcon(FEditorStyle::GetStyleSetName(), "GameProjectDialog.NoStarterContent"), LOCTEXT("NoStarterContent", "No Starter Content")));
+
+		// Only add the option to add starter content if its there to add !
+		bool bIsStarterAvailable = GameProjectUtils::IsStarterContentAvailableForNewProjects();
+		StarterContentOptions.Add(SDecoratedEnumCombo<int32>::FComboOption(
+			1, FSlateIcon(FEditorStyle::GetStyleSetName(), "GameProjectDialog.IncludeStarterContent"), LOCTEXT("IncludeStarterContent", "With Starter Content"), bIsStarterAvailable));
 
-			// Warning when enabled for mobile, since the current starter content is bad for mobile
+		TSharedRef<SWidget> Enum = SNew(SOverlay)
 			+ SOverlay::Slot()
-			//.Visibility(EVisibility::SelfHitTestInvisible)
 			.HAlign(HAlign_Right)
 			.VAlign(VAlign_Top)
 			.Padding(4)
 			[
+				// Warning when enabled for mobile, since the current starter content is bad for mobile
 				SNew(SImage)
 				.Image(FEditorStyle::GetBrush("Icons.Warning"))
 				.ToolTipText(this, &SNewProjectWizard::GetStarterContentWarningTooltip)
 				.Visibility(this, &SNewProjectWizard::GetStarterContentWarningVisibility)
 			]
-		];
+			+ SOverlay::Slot()
+			[
+				SNew(SDecoratedEnumCombo<int32>, MoveTemp(StarterContentOptions))
+				.SelectedEnum(this, &SNewProjectWizard::GetCopyStarterContentIndex)
+				.OnEnumChanged(this, &SNewProjectWizard::OnSetCopyStarterContent)
+				.Orientation(Orient_Vertical)
+			];
+
+		TSharedRef<SRichTextBlock> Description =
+			SNew(SRichTextBlock)
+			.Text(LOCTEXT("CopyStarterContent_ToolTip", "Enable to include an additional content pack containing simple placeable meshes with basic materials and textures.\nYou can opt out of including this to create a project that only has the bare essentials for the selected project template.\nYou can also add the <RichTextBlock.BoldHighlight>Starter Content</> to your project later using <RichTextBlock.BoldHighlight>Content Browser</>."))
+			.AutoWrapText(true)
+			.DecoratorStyleSet(&FEditorStyle::Get());
+		
+		AddToProjectSettingsGrid(GridPanel, Enum, Description, CurrentSlot);
+	}
 
-	ProjectSettingsOptionsBox = Widget;
+	if (!HiddenSettings.Contains(ETemplateSetting::VR))
+	{
+		TArray<SDecoratedEnumCombo<int32>::FComboOption> VREnabledOptions;
+		VREnabledOptions.Add(SDecoratedEnumCombo<int32>::FComboOption(
+			0, FSlateIcon(FEditorStyle::GetStyleSetName(), "LevelEditor.ToggleVR"/*"GameProjectDialog.VRDisabled"*/), LOCTEXT("VRDisabled", "VR Disabled")));
+
+		VREnabledOptions.Add(SDecoratedEnumCombo<int32>::FComboOption(
+			1, FSlateIcon(FEditorStyle::GetStyleSetName(), "LevelEditor.ToggleVR"/*"GameProjectDialog.VREnabled"*/), LOCTEXT("VREnabled", "VR Enabled")));
+
+		TSharedRef<SDecoratedEnumCombo<int32>> Enum = SNew(SDecoratedEnumCombo<int32>, MoveTemp(VREnabledOptions))
+			.SelectedEnum(this, &SNewProjectWizard::OnGetVREnabled)
+			.OnEnumChanged(this, &SNewProjectWizard::OnSetVREnabled)
+			.Orientation(Orient_Vertical)
+			.ToolTipText(LOCTEXT("EnableVR_Tooltip", "Enable to to add VR support to the created project."));
+
+		TSharedRef<SRichTextBlock> Description = SNew(SRichTextBlock)
+			.Text(LOCTEXT("VREnabledDescription", "Choose if VR should be enabled in the new project."))
+			.AutoWrapText(true)
+			.DecoratorStyleSet(&FEditorStyle::Get());
+
+		AddToProjectSettingsGrid(GridPanel, Enum, Description, CurrentSlot);
+	}
 
-	return ProjectSettingsOptionsBox.ToSharedRef();
+	return GridPanel;
 }
 #undef LOCTEXT_NAMESPACE
diff --git a/Engine/Source/Editor/GameProjectGeneration/Private/SNewProjectWizard.h b/Engine/Source/Editor/GameProjectGeneration/Private/SNewProjectWizard.h
--- a/Engine/Source/Editor/GameProjectGeneration/Private/SNewProjectWizard.h
+++ b/Engine/Source/Editor/GameProjectGeneration/Private/SNewProjectWizard.h
@@ -12,6 +12,7 @@
 #include "Widgets/SBoxPanel.h"
 #include "HardwareTargetingSettings.h"
 #include "Widgets/Views/STileView.h"
+#include "TemplateCategory.h"
 
 class SWizard;
 struct FTemplateItem;
@@ -24,7 +25,7 @@ class SNewProjectWizard : public SCompoundWidget
 public:
 	SLATE_BEGIN_ARGS( SNewProjectWizard ){}
 
-		SLATE_EVENT( FOnClicked, OnBackRequested )
+		SLATE_EVENT( FSimpleDelegate, OnTemplateDoubleClick )
 
 	SLATE_END_ARGS()
 
@@ -36,10 +37,22 @@ public:
 	/** Populates TemplateItemsSource with templates found on disk */
 	TMap<FName, TArray<TSharedPtr<FTemplateItem>> >& FindTemplateProjects();
 
-private:
+	/** Create the page of project settings. */
+	TSharedRef<SWidget> CreateProjectSettingsPage();
+
+	/** Handle choosing a different category tab */
+	void SetCurrentCategory(FName Category);
+
+	/** Returns true if the user is allowed to specify a project with the supplied name and path */
+	bool CanCreateProject() const;
+
+	/** Begins the creation process for the configured project */
+	void CreateAndOpenProject();
+
+	/** Should we show the project settings page? */
+	bool ShouldShowProjectSettingsPage() const;
 
-	/** Build the set of template category tabs */
-	TSharedRef<SWidget> BuildCategoryTabs();
+private:
 
 	/** Accessor for the currently selected template item */
 	TSharedPtr<FTemplateItem> GetSelectedTemplateItem() const;
@@ -79,18 +92,17 @@ private:
 	/** Accessor for the label to preview the current filename with path */
 	FString GetProjectFilenameWithPathLabelText() const;
 
-	/** Accessor for the label to show the currently selected template */
-	FText GetSelectedTemplateName() const;
-
 	/** Get the images for the selected template preview and category */
 	const FSlateBrush* GetSelectedTemplatePreviewImage() const;
-	const FSlateBrush* GetSelectedTemplateTypeImage() const;
 
 	/** Get the visibility for the selected template preview image */
 	EVisibility GetSelectedTemplatePreviewVisibility() const;
 
 	/** Get a string that details the class types referenced in the selected template */
 	FText GetSelectedTemplateClassTypes() const;
+
+	/** Get a string that details the class types referenced in the selected template */
+	FText GetSelectedTemplateLanguages() const;
 	
 	/** Get a visiblity of the class types display. If the string is empty this return Collapsed otherwise it will return Visible */
 	EVisibility GetSelectedTemplateClassVisibility() const;
@@ -104,26 +116,17 @@ private:
 	/** Gets the assembled project filename with path */
 	FString GetProjectFilenameWithPath() const;
 
-	/** Opens a web browser to the specified IDE url */
-	void OnDownloadIDEClicked(FString URL);
-
-	/** Returns true if the user is allowed to specify a project with the supplied name and path */
-	bool IsCreateProjectEnabled() const;
-
-	// Handles checking whether the specified page can be shown.
-	bool HandlePageCanShow( FName PageName ) const;
-
-	/** Fired when the page changes in the wizard */
-	void OnPageVisited(FName NewPageName);
-
 	/** Gets the visibility of the error label */
 	EVisibility GetGlobalErrorLabelVisibility() const;
 
 	/** Gets the visibility of the error label close button. Only visible when displaying a persistent error */
 	EVisibility GetGlobalErrorLabelCloseButtonVisibility() const;
 
-	/** Gets the visibility of the IDE link in the error label */
-	EVisibility GetGlobalErrorLabelIDELinkVisibility() const;
+	/** 
+	 * Gets the visibility of the location box on the template list page. 
+	 * Only visible if the selected template doesn't want to show settings.
+	 */
+	EVisibility GetTemplateListLocationBoxVisibility() const;
 
 	/** Gets the text to display in the error label */
 	FText GetGlobalErrorLabelText() const;
@@ -149,6 +152,14 @@ private:
 	/** Returns the tooltip text (actual warning) of the starter content warning */
 	FText GetStarterContentWarningTooltip() const;
 
+	int32 OnGetVREnabled() const { return bEnableVR ? 1 : 0; }
+	void OnSetVREnabled(int32 InVREnabledIndex) { bEnableVR = InVREnabledIndex == 1; }
+
+	int32 OnGetBlueprintOrCppIndex() const;
+	void OnSetBlueprintOrCppIndex(int32 Index);
+
+	/** Create the project location widget. */
+	TSharedRef<SWidget> MakeProjectLocationWidget();
 
 private:
 
@@ -169,9 +180,6 @@ private:
 	/** Returns true if we have a code template selected */
 	bool IsCompilerRequired() const;
 
-	/** Begins the creation process for the configured project */
-	void CreateAndOpenProject();
-
 	/** Opens the specified project file */
 	bool OpenProject(const FString& ProjectFile);
 
@@ -187,42 +195,27 @@ private:
 	/** Displays an error to the user */
 	void DisplayError(const FText& ErrorText);
 
-private:
-
 	/** Handler for when the Browse button is clicked */
 	FReply HandleBrowseButtonClicked( );
 
-	/** Checks whether the 'Create Project' wizard can finish. */
-	bool HandleCreateProjectWizardCanFinish( ) const;
-
-	/** Handler for when the Create button is clicked */
-	void HandleCreateProjectWizardFinished( );
-
 	/** Called when a user double-clicks a template item in the template project list. */
 	void HandleTemplateListViewDoubleClick( TSharedPtr<FTemplateItem> TemplateItem );
 
 	/** Handler for when the selection changes in the template list */
 	void HandleTemplateListViewSelectionChanged( TSharedPtr<FTemplateItem> TemplateItem, ESelectInfo::Type SelectInfo );
 
-	/** Handle choosing a different category tab */
-	void HandleCategoryChanged(ECheckBoxState Checked, FName Category);
-
-	/** Get the check state for the specified categories tab */
-	ECheckBoxState GetCategoryTabCheckState(FName Category) const;
-
-	TSharedRef<SRichTextBlock> MakeProjectSettingsDescriptionBox();
-	TSharedRef<SHorizontalBox> MakeProjectSettingsOptionsBox();
+	TSharedRef<SRichTextBlock> MakeProjectSettingsDescriptionBox() const;
+	TSharedRef<SWidget> MakeProjectSettingsOptionsBox() const;
 
 private:
 
-	/** The wizard widget */
-	TSharedPtr<class SWizard> MainWizard;
-
 	FString LastBrowsePath;
 	FString CurrentProjectFileName;
 	FString CurrentProjectFilePath;
 	FText PersistentGlobalErrorLabelText;
 
+	FSimpleDelegate OnTemplateDoubleClick;
+
 	/** The last time that the selected project file path was checked for validity. This is used to throttle I/O requests to a reasonable frequency */
 	double LastValidityCheckTime;
 
@@ -244,12 +237,15 @@ private:
 	/** True if the last NameAndLocation validity check returned that the project path is valid for creation */
 	bool bLastNameAndLocationValidityCheckSuccessful;
 
-	/** The name of the page that is currently in view */
-	FName CurrentPageName;
-
 	/** True if user has selected to copy starter content. */
 	bool bCopyStarterContent;
 
+	/** Whether or not to enable VR in the created project. */
+	bool bEnableVR;
+
+	/** Whether or not we should use the blueprint or C++ version of this template. */
+	bool bShouldGenerateCode;
+
 	/** Name of the currently selected category */
 	FName ActiveCategory;
 
@@ -260,14 +256,14 @@ private:
 	TArray<TSharedPtr<FTemplateItem> > FilteredTemplateList;
 
 	/** The slate widget representing the list of templates */
-	TSharedPtr<STileView<TSharedPtr<FTemplateItem> > > TemplateListView;
+	TSharedPtr<STileView<TSharedPtr<FTemplateItem>>> TemplateListView;
 
 	/** Names for pages */
 	static FName TemplatePageName;
 	static FName NameAndLocationPageName;
 
 	TSharedPtr<SRichTextBlock> ProjectSettingsDescriptionBox;
-	TSharedPtr<SHorizontalBox> ProjectSettingsOptionsBox;
+	TSharedPtr<SWidget> ProjectSettingsOptionsBox;
 
 	bool bProjectSettingsHidden;
 };
diff --git a/Engine/Source/Editor/GameProjectGeneration/Private/Tests/GameProjectAutomationTests.cpp b/Engine/Source/Editor/GameProjectGeneration/Private/Tests/GameProjectAutomationTests.cpp
--- a/Engine/Source/Editor/GameProjectGeneration/Private/Tests/GameProjectAutomationTests.cpp
+++ b/Engine/Source/Editor/GameProjectGeneration/Private/Tests/GameProjectAutomationTests.cpp
@@ -78,7 +78,7 @@ namespace GameProjectAutomationUtils
 	 * @param	OutMatchedProjects	Total projects matching criteria
 	 * @param	OutCreatedProjects	Total projects succesfully created
 	 */
-	static void CreateProjectSet(TMap<FName, TArray<TSharedPtr<FTemplateItem>> >& InTemplates, EHardwareClass::Type InTargetedHardware, EGraphicsPreset::Type InGraphicPreset, EContentSourceCategory InCategory, bool bInCopyStarterContent, int32 &OutCreatedProjects, int32 &OutMatchedProjects)
+	static void CreateProjectSet(TMap<FName, TArray<TSharedPtr<FTemplateItem>> >& InTemplates, EHardwareClass::Type InTargetedHardware, EGraphicsPreset::Type InGraphicPreset, FName InCategory, bool bInCopyStarterContent, int32 &OutCreatedProjects, int32 &OutMatchedProjects)
 	{		
 		// If this is empty, it will use the same name for each project, otherwise it will create a project based on target platform and source template
 		FString TestRootFolder;// = "ProjectTests";
@@ -87,109 +87,88 @@ namespace GameProjectAutomationUtils
 		bool bRemoveCreatedProjects = true;		
 		OutCreatedProjects = 0;
 		OutMatchedProjects = 0;
-		UEnum* SourceCategoryEnum = StaticEnum<EContentSourceCategory>();
-
-		// The category name in the FTemplateItem is not the same as the enum definition EContentSourceCategory - convert it
-		FName CategoryName;
-		if( InCategory == EContentSourceCategory::BlueprintFeature)
-		{
-			CategoryName = FTemplateCategory::BlueprintCategoryName;
-		}
-		else if (InCategory == EContentSourceCategory::CodeFeature)
-		{
-			CategoryName = FTemplateCategory::CodeCategoryName;
-		}
-		else if (InCategory == EContentSourceCategory::EnterpriseFeature)
-		{
-			CategoryName = FTemplateCategory::EnterpriseCategoryName;
-		}
-		else
-		{
-			// We didn't match a category
-			if (SourceCategoryEnum != nullptr)
-			{
-				UE_LOG(LogGameProjectGenerationTests, Display, TEXT("Test failed! Unknown category type %s"), *SourceCategoryEnum->GetNameStringByValue(int64(InCategory)));
-			}
-			else
-			{
-				UE_LOG(LogGameProjectGenerationTests, Display, TEXT("Test failed! Unknown category type %d"), (int32)(InCategory));
-			}
-			
-			return;
-		}
 
 		// Iterate all templates and try to create those that match the required criteria
-		for (auto EachTemplate : InTemplates)
+		for (auto& EachTemplate : InTemplates)
 		{
-			FName Name = EachTemplate.Key;						
 			// Check this is the correct category
-			if (Name == CategoryName)
+			if (EachTemplate.Key == InCategory)
 			{
 				// Now iterate each template in the category
 				for (TSharedPtr<FTemplateItem> OneTemplate : EachTemplate.Value)
-				{				
+				{
 					TSharedPtr<FTemplateItem> Item = OneTemplate;
 
-					// If this template is flagged as not for creation don't try to create it
-					if( Item->ProjectFile.IsEmpty())
-						continue;
-
-					FString DesiredProjectFilename;
-					if( TestRootFolder.IsEmpty() == true)
+					// loop over the code project file and the blueprint file
+					FString ProjectFile = Item->BlueprintProjectFile;
+					for (int i = 0; i < 2; ++i)
 					{
-						// Same name for all
-						DesiredProjectFilename = GameProjectAutomationUtils::GetDesiredProjectFilename();
-					}
-					else
-					{
-						// Unique names
-						FString Hardware;
-						if (InTargetedHardware == EHardwareClass::Desktop)
+						if (ProjectFile.IsEmpty())
+							continue;
+
+						FString DesiredProjectFilename;
+						if (TestRootFolder.IsEmpty() == true)
 						{
-							Hardware = TEXT("Dsk");
+							// Same name for all
+							DesiredProjectFilename = GameProjectAutomationUtils::GetDesiredProjectFilename();
 						}
 						else
 						{
-							Hardware = TEXT("Mob");
+							// Unique names
+							FString Hardware;
+							if (InTargetedHardware == EHardwareClass::Desktop)
+							{
+								Hardware = TEXT("Dsk");
+							}
+							else
+							{
+								Hardware = TEXT("Mob");
+							}
+							FString ProjectName = FPaths::GetCleanFilename(ProjectFile).Replace(TEXT("TP_"), TEXT(""));
+							FString ProjectDirName = FPaths::GetBaseFilename(ProjectFile).Replace(TEXT("TP_"), TEXT(""));
+							FString BasePath = FPaths::RootDir();
+							DesiredProjectFilename = FString::Printf(TEXT("%s/%s/%s%s/%s%s"), *BasePath, *TestRootFolder, *Hardware, *ProjectDirName, *Hardware, *ProjectName);
 						}
-						FString ProjectName = FPaths::GetCleanFilename(Item->ProjectFile).Replace(TEXT("TP_"), TEXT(""));
-						FString ProjectDirName = FPaths::GetBaseFilename(Item->ProjectFile).Replace(TEXT("TP_"), TEXT(""));
-						FString BasePath = FPaths::RootDir();
-						DesiredProjectFilename = FString::Printf(TEXT("%s/%s/%s%s/%s%s"), *BasePath, *TestRootFolder, *Hardware, *ProjectDirName, *Hardware, *ProjectName);
-					}
 
-					// If the project already exists, delete it just in case things were left in a bad state.
-					if ( IFileManager::Get().DirectoryExists(*FPaths::GetPath(DesiredProjectFilename)) )
-					{
-						IFileManager::Get().DeleteDirectory(*FPaths::GetPath(DesiredProjectFilename), /*RequireExists=*/false, /*Tree=*/true);
-					}
-					
-					// Setup creation parameters
-					FText FailReason, FailLog;
-					FProjectInformation ProjectInfo(DesiredProjectFilename, Item->bGenerateCode, bInCopyStarterContent, Item->ProjectFile);
-					ProjectInfo.TargetedHardware = InTargetedHardware;
-					ProjectInfo.DefaultGraphicsPerformance = InGraphicPreset;
-					TArray<FString> CreatedFiles;
-					OutMatchedProjects++;
-
-					// Finally try to create the project					
-					if (!GameProjectUtils::CreateProject(ProjectInfo, FailReason, FailLog, &CreatedFiles))
-					{
-						// Failed, report the reason
-						UE_LOG(LogGameProjectGenerationTests, Error, TEXT("Failed to create %s project %s based on %s. Reason: %s\nProject Creation Failure Log:\n%s"), *SourceCategoryEnum->GetNameStringByValue((int64)InCategory), *DesiredProjectFilename, *Item->Name.ToString(), *FailReason.ToString(), *FailLog.ToString());
-					}
-					else
-					{
-						// Created ok
-						OutCreatedProjects++;
-						
-						// Now remove the files we just created (if required)
-						if(bRemoveCreatedProjects == true) //-V547
+						// If the project already exists, delete it just in case things were left in a bad state.
+						if (IFileManager::Get().DirectoryExists(*FPaths::GetPath(DesiredProjectFilename)))
+						{
+							IFileManager::Get().DeleteDirectory(*FPaths::GetPath(DesiredProjectFilename), /*RequireExists=*/false, /*Tree=*/true);
+						}
+
+						// Setup creation parameters
+						FText FailReason, FailLog;
+						FProjectInformation ProjectInfo;
+						ProjectInfo.ProjectFilename = DesiredProjectFilename;
+						ProjectInfo.bShouldGenerateCode = ProjectFile == Item->CodeProjectFile;
+						ProjectInfo.bCopyStarterContent = bInCopyStarterContent;
+						ProjectInfo.TemplateFile = ProjectFile;
+						ProjectInfo.TargetedHardware = InTargetedHardware;
+						ProjectInfo.DefaultGraphicsPerformance = InGraphicPreset;
+						TArray<FString> CreatedFiles;
+						OutMatchedProjects++;
+
+						// Finally try to create the project					
+						if (!GameProjectUtils::CreateProject(ProjectInfo, FailReason, FailLog, &CreatedFiles))
 						{
-							FString RootFolder = FPaths::GetPath(DesiredProjectFilename);
-							GameProjectUtils::DeleteCreatedFiles(RootFolder, CreatedFiles);
+							// Failed, report the reason
+							UE_LOG(LogGameProjectGenerationTests, Error, TEXT("Failed to create %s project %s based on %s. Reason: %s\nProject Creation Failure Log:\n%s"), *InCategory.ToString(), *DesiredProjectFilename, *Item->Name.ToString(), *FailReason.ToString(), *FailLog.ToString());
+						}
+						else
+						{
+							// Created ok
+							OutCreatedProjects++;
+
+							// Now remove the files we just created (if required)
+							if (bRemoveCreatedProjects == true) //-V547
+							{
+								FString RootFolder = FPaths::GetPath(DesiredProjectFilename);
+								GameProjectUtils::DeleteCreatedFiles(RootFolder, CreatedFiles);
+							}
 						}
 					}
+
+					ProjectFile = Item->CodeProjectFile;
 				}
 			}
 		}
@@ -200,7 +179,7 @@ namespace GameProjectAutomationUtils
 /**
 * Automation test to clean up old test project files
 */
-IMPLEMENT_SIMPLE_AUTOMATION_TEST(FBuildPromotionNewProjectCleanupTest, "System.Promotion.Project Promotion Pass.Step 1 Blank Project Creation.Cleanup Potential Project Location", EAutomationTestFlags::Disabled | EAutomationTestFlags::EditorContext | EAutomationTestFlags::EngineFilter );
+IMPLEMENT_SIMPLE_AUTOMATION_TEST(FBuildPromotionNewProjectCleanupTest, "System.Promotion.Project Promotion Pass.Step 1 Blank Project Creation.Cleanup Potential Project Location", /*EAutomationTestFlags::Disabled |*/ EAutomationTestFlags::EditorContext | EAutomationTestFlags::EngineFilter );
 bool FBuildPromotionNewProjectCleanupTest::RunTest(const FString& Parameters)
 {
 	FString DesiredProjectFilename = GameProjectAutomationUtils::GetDesiredProjectFilename();
@@ -235,7 +214,7 @@ bool FBuildPromotionNewProjectCleanupTest::RunTest(const FString& Parameters)
 /**
 * Automation test to create a new project
 */
-IMPLEMENT_SIMPLE_AUTOMATION_TEST(FBuildPromotionNewProjectCreateTest, "System.Promotion.Project Promotion Pass.Step 1 Blank Project Creation.Create Project", EAutomationTestFlags::Disabled | EAutomationTestFlags::EditorContext | EAutomationTestFlags::EngineFilter | EAutomationTestFlags::RequiresUser);
+IMPLEMENT_SIMPLE_AUTOMATION_TEST(FBuildPromotionNewProjectCreateTest, "System.Promotion.Project Promotion Pass.Step 1 Blank Project Creation.Create Project", /*EAutomationTestFlags::Disabled |*/ EAutomationTestFlags::EditorContext | EAutomationTestFlags::EngineFilter | EAutomationTestFlags::RequiresUser);
 bool FBuildPromotionNewProjectCreateTest::RunTest(const FString& Parameters)
 {
 	FString DesiredProjectFilename = GameProjectAutomationUtils::GetDesiredProjectFilename();
@@ -262,7 +241,10 @@ bool FBuildPromotionNewProjectCreateTest::RunTest(const FString& Parameters)
 	}
 
 	FText FailReason, FailLog;
-	if (GameProjectUtils::CreateProject(FProjectInformation(DesiredProjectFilename, false, true), FailReason, FailLog))
+	FProjectInformation ProjectInfo;
+	ProjectInfo.ProjectFilename = DesiredProjectFilename;
+	ProjectInfo.bCopyStarterContent = true;
+	if (GameProjectUtils::CreateProject(ProjectInfo, FailReason, FailLog))
 	{
 		UE_LOG(LogGameProjectGenerationTests, Display, TEXT("Generated a new project: %s"), *DesiredProjectFilename);
 		UE_LOG(LogGameProjectGenerationTests, Display, TEXT("Test successful!"));
@@ -280,7 +262,7 @@ bool FBuildPromotionNewProjectCreateTest::RunTest(const FString& Parameters)
 /**
 * Automation test to create a simple level and save it
 */
-IMPLEMENT_SIMPLE_AUTOMATION_TEST(FBuildPromotionNewProjectMapTest, "System.Promotion.Project Promotion Pass.Step 2 Basic Level Creation.Create Basic Level", EAutomationTestFlags::Disabled | EAutomationTestFlags::EditorContext | EAutomationTestFlags::EngineFilter);
+IMPLEMENT_SIMPLE_AUTOMATION_TEST(FBuildPromotionNewProjectMapTest, "System.Promotion.Project Promotion Pass.Step 2 Basic Level Creation.Create Basic Level", /*EAutomationTestFlags::Disabled |*/ EAutomationTestFlags::EditorContext | EAutomationTestFlags::EngineFilter);
 bool FBuildPromotionNewProjectMapTest::RunTest(const FString& Parameters)
 {
 	//New level
@@ -336,7 +318,7 @@ bool FBuildPromotionNewProjectMapTest::RunTest(const FString& Parameters)
 /*
 * Template project creation test
 */
-IMPLEMENT_SIMPLE_AUTOMATION_TEST(FCreateBPTemplateProjectAutomationTests, "System.Promotion.Project Promotion Pass.Step 3 NewProjectCreationTests.CreateBlueprintProjects", EAutomationTestFlags::EditorContext | EAutomationTestFlags::EngineFilter | EAutomationTestFlags::Disabled)
+IMPLEMENT_SIMPLE_AUTOMATION_TEST(FCreateBPTemplateProjectAutomationTests, "System.Promotion.Project Promotion Pass.Step 3 NewProjectCreationTests.CreateBlueprintProjects", EAutomationTestFlags::EditorContext | EAutomationTestFlags::EngineFilter /*| EAutomationTestFlags::Disabled*/)
 
 /** 
  * Uses the new project wizard to locate all templates available for new blueprint project creation and verifies creation succeeds.
@@ -352,11 +334,11 @@ bool FCreateBPTemplateProjectAutomationTests::RunTest(const FString& Parameters)
 	TMap<FName, TArray<TSharedPtr<FTemplateItem>> >& Templates = NewProjectWizard->FindTemplateProjects();
 	int32 OutMatchedProjectsDesk = 0;
 	int32 OutCreatedProjectsDesk = 0;
-	GameProjectAutomationUtils::CreateProjectSet(Templates, EHardwareClass::Desktop, EGraphicsPreset::Maximum, EContentSourceCategory::BlueprintFeature, false, OutMatchedProjectsDesk, OutCreatedProjectsDesk);
+	GameProjectAutomationUtils::CreateProjectSet(Templates, EHardwareClass::Desktop, EGraphicsPreset::Maximum, "Game", false, OutMatchedProjectsDesk, OutCreatedProjectsDesk);
 
 	int32 OutMatchedProjectsMob = 0;
 	int32 OutCreatedProjectsMob = 0;
-	GameProjectAutomationUtils::CreateProjectSet(Templates, EHardwareClass::Mobile, EGraphicsPreset::Maximum, EContentSourceCategory::BlueprintFeature, false, OutMatchedProjectsMob, OutCreatedProjectsMob);
+	GameProjectAutomationUtils::CreateProjectSet(Templates, EHardwareClass::Mobile, EGraphicsPreset::Maximum, "Visualization", false, OutMatchedProjectsMob, OutCreatedProjectsMob);
 	
 	return ( OutMatchedProjectsDesk == OutCreatedProjectsDesk ) && ( OutMatchedProjectsMob == OutCreatedProjectsMob  );
 }
@@ -364,7 +346,7 @@ bool FCreateBPTemplateProjectAutomationTests::RunTest(const FString& Parameters)
 /*
 * Template project creation test
 */
-IMPLEMENT_SIMPLE_AUTOMATION_TEST(FCreateCPPTemplateProjectAutomationTests, "System.Promotion.Project Promotion Pass.Step 3 NewProjectCreationTests.CreateCodeProjects", EAutomationTestFlags::EditorContext | EAutomationTestFlags::EngineFilter | EAutomationTestFlags::Disabled)
+IMPLEMENT_SIMPLE_AUTOMATION_TEST(FCreateCPPTemplateProjectAutomationTests, "System.Promotion.Project Promotion Pass.Step 3 NewProjectCreationTests.CreateCodeProjects", EAutomationTestFlags::EditorContext | EAutomationTestFlags::EngineFilter /*| EAutomationTestFlags::Disabled*/)
 
 /** 
  * Uses the new project wizard to locate all templates available for new code project creation and verifies creation succeeds.
@@ -380,11 +362,11 @@ bool FCreateCPPTemplateProjectAutomationTests::RunTest(const FString& Parameters
 	TMap<FName, TArray<TSharedPtr<FTemplateItem>> >& Templates = NewProjectWizard->FindTemplateProjects();//GameProjectAutomationUtils::CreateTemplateList();
 	int32 OutMatchedProjectsDesk = 0;
 	int32 OutCreatedProjectsDesk = 0;
-	GameProjectAutomationUtils::CreateProjectSet(Templates, EHardwareClass::Desktop, EGraphicsPreset::Maximum, EContentSourceCategory::CodeFeature, false, OutMatchedProjectsDesk, OutCreatedProjectsDesk);
+	GameProjectAutomationUtils::CreateProjectSet(Templates, EHardwareClass::Desktop, EGraphicsPreset::Maximum, "Game", false, OutMatchedProjectsDesk, OutCreatedProjectsDesk);
 
 	int32 OutMatchedProjectsMob = 0;
 	int32 OutCreatedProjectsMob = 0;
-	GameProjectAutomationUtils::CreateProjectSet(Templates, EHardwareClass::Mobile, EGraphicsPreset::Maximum, EContentSourceCategory::CodeFeature, false, OutMatchedProjectsMob, OutCreatedProjectsMob);
+	GameProjectAutomationUtils::CreateProjectSet(Templates, EHardwareClass::Mobile, EGraphicsPreset::Maximum, "Game", false, OutMatchedProjectsMob, OutCreatedProjectsMob);
 
 	return (OutMatchedProjectsDesk == OutCreatedProjectsDesk) && (OutMatchedProjectsMob == OutCreatedProjectsMob);
 
