Missed change causing compile error.

diff --git a/Engine/Source/Editor/GameProjectGeneration/Private/SNewProjectWizard.cpp b/Engine/Source/Editor/GameProjectGeneration/Private/SNewProjectWizard.cpp
--- a/Engine/Source/Editor/GameProjectGeneration/Private/SNewProjectWizard.cpp
+++ b/Engine/Source/Editor/GameProjectGeneration/Private/SNewProjectWizard.cpp
@@ -1153,7 +1153,7 @@ bool SNewProjectWizard::CreateProject( const FString& ProjectFile )
 	ProjectInfo.bEnableVR = bEnableVR;
 
 	// @todo (sebastiann): This is a poor way to determine Enterprise-ness.
-	ProjectInfo.bIsEnterpriseProject = SelectedTemplate->Category != TEXT("Game");
+	ProjectInfo.bIsEnterpriseProject = ActiveCategory != TEXT("Game");
 
 	if (!GameProjectUtils::CreateProject(ProjectInfo, FailReason, FailLog))
 	{
