The project language dropdown on the Project Settings page of the New Project wizard was not being updated correctly when changing selections, but rather only taking the values given on startup. Now it gets updated when changing the template selection.

diff --git a/Engine/Source/Editor/GameProjectGeneration/Private/SGameProjectDialog.cpp b/Engine/Source/Editor/GameProjectGeneration/Private/SGameProjectDialog.cpp
--- a/Engine/Source/Editor/GameProjectGeneration/Private/SGameProjectDialog.cpp
+++ b/Engine/Source/Editor/GameProjectGeneration/Private/SGameProjectDialog.cpp
@@ -63,9 +63,7 @@ void SGameProjectDialog::Construct(const FArguments& InArgs, EMode InMode)
 
 	NewProjectWizard = SNew(SNewProjectWizard)
 		.OnTemplateDoubleClick(this, &SGameProjectDialog::OnTemplateDoubleClick);
-	ProjectSettingsPage = NewProjectWizard->CreateProjectSettingsPage();
 	ProjectBrowserPage = SNew(SProjectBrowser);
-	LandingPage = CreateLandingPage();
 
 	const float UniformPadding = 16.0f;
 
@@ -96,7 +94,7 @@ void SGameProjectDialog::Construct(const FArguments& InArgs, EMode InMode)
 			.Name(GetPageTitle(GameProjectDialogDefs::LandingPageIndex))
 			.CanShow(DialogMode != EMode::Open)
 			[
-				LandingPage.ToSharedRef()
+				CreateLandingPage()
 			]
 
 			+ SWizard::Page()
@@ -117,7 +115,7 @@ void SGameProjectDialog::Construct(const FArguments& InArgs, EMode InMode)
 			.Name(GetPageTitle(GameProjectDialogDefs::ProjectSettingsPageIndex))
 			.CanShow(DialogMode != EMode::Open)
 			[
-				ProjectSettingsPage.ToSharedRef()
+				NewProjectWizard->CreateProjectSettingsPage()
 			]
 		]
 	];
diff --git a/Engine/Source/Editor/GameProjectGeneration/Private/SGameProjectDialog.h b/Engine/Source/Editor/GameProjectGeneration/Private/SGameProjectDialog.h
--- a/Engine/Source/Editor/GameProjectGeneration/Private/SGameProjectDialog.h
+++ b/Engine/Source/Editor/GameProjectGeneration/Private/SGameProjectDialog.h
@@ -50,8 +50,6 @@ private:
 	TSharedPtr<SProjectBrowser> ProjectBrowserPage;
 	TSharedPtr<SNewProjectWizard> NewProjectWizard;
 	TSharedPtr<SRecentProjectBrowser> RecentProjectBrowser;
-	TSharedPtr<SWidget> LandingPage;
-	TSharedPtr<SWidget> ProjectSettingsPage;
 
 	TSharedPtr<STileView<TSharedPtr<FTemplateCategory>>> MajorCategoryTileView;
 	TSharedPtr<STileView<TSharedPtr<FTemplateCategory>>> MinorCategoryTileView;
diff --git a/Engine/Source/Editor/GameProjectGeneration/Private/SNewProjectWizard.cpp b/Engine/Source/Editor/GameProjectGeneration/Private/SNewProjectWizard.cpp
--- a/Engine/Source/Editor/GameProjectGeneration/Private/SNewProjectWizard.cpp
+++ b/Engine/Source/Editor/GameProjectGeneration/Private/SNewProjectWizard.cpp
@@ -40,7 +40,6 @@
 #include "TemplateCategory.h"
 #include "Widgets/SToolTip.h"
 #include "Widgets/Workflow/SWizard.h"
-#include "SDecoratedEnumCombo.h"
 #include "IDocumentation.h"
 #include "Internationalization/BreakIterator.h"
 #include "Dialogs/SOutputLogDialog.h"
@@ -143,7 +142,7 @@ public:
 				.VAlign(VAlign_Center)
 				[
 					SNew(SBox)
-					.HeightOverride(3)
+					.HeightOverride(5)
 					[
 						SNew(SBorder)
 						.BorderImage(FEditorStyle::GetBrush("FilePath.GroupIndicator"))
@@ -189,7 +188,7 @@ public:
 				.VAlign(VAlign_Center)
 				[
 					SNew(SBox)
-					.HeightOverride(3)
+					.HeightOverride(5)
 					[
 						SNew(SBorder)
 						.BorderImage(FEditorStyle::GetBrush("FilePath.GroupIndicator"))
@@ -349,7 +348,7 @@ void SNewProjectWizard::Construct( const FArguments& InArgs )
 			.Text(LOCTEXT("ProjectTemplateDescription", "Choose a <RichTextBlock.BoldHighlight>template</> to use as a starting point for your new project.  Any of these features can be added later by clicking <RichTextBlock.BoldHighlight>Add Feature or Content Pack</> in <RichTextBlock.BoldHighlight>Content Browser</>."))
 			.AutoWrapText(true)
 			.DecoratorStyleSet(&FEditorStyle::Get())
-			.ToolTip(IDocumentation::Get()->CreateToolTip(LOCTEXT("TemplateChoiceTooltip", "A template consists of a little bit of player control logic (either as a Blueprint or in C++), input bindings, and appropriate prototyping assets."), NULL, TEXT("Shared/Editor/NewProjectWizard"), TEXT("TemplateChoice")))
+			.ToolTip(IDocumentation::Get()->CreateToolTip(LOCTEXT("TemplateChoiceTooltip", "A template consists of a little bit of player control logic (either as a Blueprint or in C++), input bindings, and appropriate prototyping assets."), nullptr, TEXT("Shared/Editor/NewProjectWizard"), TEXT("TemplateChoice")))
 		]
 
 		+ SVerticalBox::Slot()
@@ -571,8 +570,24 @@ void SNewProjectWizard::HandleTemplateListViewSelectionChanged(TSharedPtr<FTempl
 {
 	UpdateProjectFileValidity();
 
-	FString CodeProjectFile = GetSelectedTemplateProperty(&FTemplateItem::CodeProjectFile);
-	bShouldGenerateCode = CodeProjectFile.IsEmpty() == false;
+	bool bIsBlueprintAvailable = !GetSelectedTemplateProperty(&FTemplateItem::BlueprintProjectFile).IsEmpty();
+	bool bIsCodeAvailable = !GetSelectedTemplateProperty(&FTemplateItem::CodeProjectFile).IsEmpty();
+
+	// if neither is available, then this is a blank template, so both are available
+	if (!bIsBlueprintAvailable && !bIsCodeAvailable)
+	{
+		bIsBlueprintAvailable = true;
+		bIsCodeAvailable = true;
+	}
+
+	TArray<SDecoratedEnumCombo<int32>::FComboOption>& LanguageOptions = ProjectLanguageEnum->GetOptions();
+	ensure(LanguageOptions[0].Text.IdenticalTo(LOCTEXT("ProjectDialog_Blueprint", "Blueprint")));
+	ensure(LanguageOptions[1].Text.IdenticalTo(LOCTEXT("ProjectDialog_Code", "C++")));
+
+	LanguageOptions[0].bChoosable = bIsBlueprintAvailable;
+	LanguageOptions[1].bChoosable = bIsCodeAvailable;
+
+	bShouldGenerateCode = !bIsBlueprintAvailable;
 }
 
 TSharedPtr<FTemplateItem> SNewProjectWizard::GetSelectedTemplateItem() const
@@ -583,7 +598,7 @@ TSharedPtr<FTemplateItem> SNewProjectWizard::GetSelectedTemplateItem() const
 		return SelectedItems[0];
 	}
 	
-	return NULL;
+	return nullptr;
 }
 
 FText SNewProjectWizard::GetSelectedTemplateClassTypes() const
@@ -980,7 +995,6 @@ TMap<FName, TArray<TSharedPtr<FTemplateItem>> >& SNewProjectWizard::FindTemplate
 	return Templates;
 }
 
-
 void SNewProjectWizard::SetDefaultProjectLocation( )
 {
 	FString DefaultProjectFilePath;
@@ -1045,7 +1059,6 @@ void SNewProjectWizard::SetDefaultProjectLocation( )
 	}
 }
 
-
 void SNewProjectWizard::UpdateProjectFileValidity( )
 {
 	// Global validity
@@ -1179,7 +1192,6 @@ bool SNewProjectWizard::CreateProject( const FString& ProjectFile )
 	return true;
 }
 
-
 void SNewProjectWizard::CreateAndOpenProject( )
 {
 	if( !CanCreateProject() )
@@ -1223,7 +1235,6 @@ void SNewProjectWizard::CreateAndOpenProject( )
 	}
 }
 
-
 bool SNewProjectWizard::OpenProject( const FString& ProjectFile )
 {
 	FText FailReason;
@@ -1239,7 +1250,6 @@ bool SNewProjectWizard::OpenProject( const FString& ProjectFile )
 	return false;
 }
 
-
 bool SNewProjectWizard::OpenCodeIDE( const FString& ProjectFile )
 {
 	FText FailReason;
@@ -1256,7 +1266,6 @@ bool SNewProjectWizard::OpenCodeIDE( const FString& ProjectFile )
 	return false;
 }
 
-
 void SNewProjectWizard::CloseWindowIfAppropriate( bool ForceClose )
 {
 	if ( ForceClose || FApp::HasProjectName() )
@@ -1270,7 +1279,6 @@ void SNewProjectWizard::CloseWindowIfAppropriate( bool ForceClose )
 	}
 }
 
-
 void SNewProjectWizard::DisplayError( const FText& ErrorText )
 {
 	FString ErrorString = ErrorText.ToString();
@@ -1285,7 +1293,6 @@ void SNewProjectWizard::DisplayError( const FText& ErrorText )
 	}
 }
 
-
 /* SNewProjectWizard event handlers
  *****************************************************************************/
 
@@ -1323,7 +1330,7 @@ TSharedRef<SWidget> SNewProjectWizard::MakeProjectLocationWidget()
 			.Text(LOCTEXT("ProjectPathDescription", "Select a <RichTextBlock.BoldHighlight>location</> for your project to be stored."))
 			.AutoWrapText(true)
 			.DecoratorStyleSet(&FEditorStyle::Get())
-			.ToolTip(IDocumentation::Get()->CreateToolTip(LOCTEXT("ProjectPathDescriptionTooltip", "All of your project content and code will be stored here."), NULL, TEXT("Shared/Editor/NewProjectWizard"), TEXT("ProjectPath")))
+			.ToolTip(IDocumentation::Get()->CreateToolTip(LOCTEXT("ProjectPathDescriptionTooltip", "All of your project content and code will be stored here."), nullptr, TEXT("Shared/Editor/NewProjectWizard"), TEXT("ProjectPath")))
 		]
 
 		+ SVerticalBox::Slot()
@@ -1346,9 +1353,6 @@ TSharedRef<SWidget> SNewProjectWizard::CreateProjectSettingsPage()
 {
 	const float UniformPadding = 16.f;
 
-	ProjectSettingsDescriptionBox = MakeProjectSettingsDescriptionBox();
-	ProjectSettingsOptionsBox = MakeProjectSettingsOptionsBox();
-
 	TSharedRef<SWidget> PageWidget = SNew(SOverlay)
 	+ SOverlay::Slot()
 	[
@@ -1357,7 +1361,7 @@ TSharedRef<SWidget> SNewProjectWizard::CreateProjectSettingsPage()
 		.AutoHeight()
 		.Padding(0, 8, 0, 8)
 		[
-			ProjectSettingsDescriptionBox.ToSharedRef()
+			MakeProjectSettingsDescriptionBox()
 		]
 		+ SVerticalBox::Slot()
 		.Padding(0)
@@ -1375,7 +1379,7 @@ TSharedRef<SWidget> SNewProjectWizard::CreateProjectSettingsPage()
 					+ SVerticalBox::Slot()
 					.Padding(FMargin(0, 0, 0, UniformPadding))
 					[
-						ProjectSettingsOptionsBox.ToSharedRef()
+						MakeProjectSettingsOptionsBox()
 					]
 				]
 			]
@@ -1508,7 +1512,7 @@ TSharedRef<SRichTextBlock> SNewProjectWizard::MakeProjectSettingsDescriptionBox(
 		.Text(LOCTEXT("ProjectSettingsDescription", "Choose some <RichTextBlock.BoldHighlight>settings</> for your project."))
 		.AutoWrapText(true)
 		.DecoratorStyleSet(&FEditorStyle::Get())
-		.ToolTip(IDocumentation::Get()->CreateToolTip(LOCTEXT("HardwareTargetTooltip", "These settings will choose good defaults for a number of other settings in the project such as post-processing flags and touch input emulation using the mouse."), NULL, TEXT("Shared/Editor/NewProjectWizard"), TEXT("TargetHardware")));
+		.ToolTip(IDocumentation::Get()->CreateToolTip(LOCTEXT("HardwareTargetTooltip", "These settings will choose good defaults for a number of other settings in the project such as post-processing flags and touch input emulation using the mouse."), nullptr, TEXT("Shared/Editor/NewProjectWizard"), TEXT("TargetHardware")));
 
 	return Widget;
 }
@@ -1537,7 +1541,7 @@ static void AddToProjectSettingsGrid(TSharedRef<SGridPanel> Grid, const TSharedR
 	}
 }
 
-TSharedRef<SWidget> SNewProjectWizard::MakeProjectSettingsOptionsBox() const
+TSharedRef<SWidget> SNewProjectWizard::MakeProjectSettingsOptionsBox()
 {
 	static const int EnumWidth = 160;
 
@@ -1553,28 +1557,20 @@ TSharedRef<SWidget> SNewProjectWizard::MakeProjectSettingsOptionsBox() const
 
 	if (!HiddenSettings.Contains(ETemplateSetting::Languages))
 	{
-		bool bIsCodeAvailable = !GetSelectedTemplateProperty(&FTemplateItem::CodeProjectFile).IsEmpty();
-		bool bIsBlueprintAvailable = !GetSelectedTemplateProperty(&FTemplateItem::BlueprintProjectFile).IsEmpty();
-
-		// if neither is available, then this is a blank template, so both are available
-		if (bIsCodeAvailable && bIsBlueprintAvailable)
-		{
-			bIsCodeAvailable = true;
-			bIsBlueprintAvailable = true;
-		}
-
 		TArray<SDecoratedEnumCombo<int32>::FComboOption> BlueprintOrCppOptions;
 		BlueprintOrCppOptions.Add(SDecoratedEnumCombo<int32>::FComboOption(
-			0, FSlateIcon(FEditorStyle::GetStyleSetName(), "GameProjectDialog.BlueprintImage_64"), LOCTEXT("ProjectDialog_Blueprint", "Blueprint"), bIsBlueprintAvailable));
+			0, FSlateIcon(FEditorStyle::GetStyleSetName(), "GameProjectDialog.BlueprintImage_64"), LOCTEXT("ProjectDialog_Blueprint", "Blueprint")));
 
 		BlueprintOrCppOptions.Add(SDecoratedEnumCombo<int32>::FComboOption(
-			1, FSlateIcon(FEditorStyle::GetStyleSetName(), "GameProjectDialog.CodeImage_64"), LOCTEXT("ProjectDialog_Code", "C++"), bIsCodeAvailable));
+			1, FSlateIcon(FEditorStyle::GetStyleSetName(), "GameProjectDialog.CodeImage_64"), LOCTEXT("ProjectDialog_Code", "C++")));
 
 		TSharedRef<SDecoratedEnumCombo<int32>> Enum = SNew(SDecoratedEnumCombo<int32>, MoveTemp(BlueprintOrCppOptions))
 			.SelectedEnum(this, &SNewProjectWizard::OnGetBlueprintOrCppIndex)
 			.OnEnumChanged(this, &SNewProjectWizard::OnSetBlueprintOrCppIndex)
 			.Orientation(Orient_Vertical);
 
+		ProjectLanguageEnum = Enum;
+
 		TSharedRef<SRichTextBlock> Description = SNew(SRichTextBlock)
 			.Text(LOCTEXT("ProjectDialog_BlueprintOrCppDescription", "Choose whether to create a Blueprint or C++ project."))
 			.AutoWrapText(true)
@@ -1598,7 +1594,6 @@ TSharedRef<SWidget> SNewProjectWizard::MakeProjectSettingsOptionsBox() const
 		AddToProjectSettingsGrid(GridPanel, Enum, Description, CurrentSlot);
 	}
 
-
 	if (!HiddenSettings.Contains(ETemplateSetting::GraphicsPreset))
 	{
 		TSharedRef<SWidget> Enum = HardwareTargeting.MakeGraphicsPresetTargetCombo(
diff --git a/Engine/Source/Editor/GameProjectGeneration/Private/SNewProjectWizard.h b/Engine/Source/Editor/GameProjectGeneration/Private/SNewProjectWizard.h
--- a/Engine/Source/Editor/GameProjectGeneration/Private/SNewProjectWizard.h
+++ b/Engine/Source/Editor/GameProjectGeneration/Private/SNewProjectWizard.h
@@ -12,6 +12,7 @@
 #include "Widgets/SBoxPanel.h"
 #include "HardwareTargetingSettings.h"
 #include "Widgets/Views/STileView.h"
+#include "SDecoratedEnumCombo.h"
 #include "TemplateCategory.h"
 
 class SWizard;
@@ -205,7 +206,7 @@ private:
 	void HandleTemplateListViewSelectionChanged( TSharedPtr<FTemplateItem> TemplateItem, ESelectInfo::Type SelectInfo );
 
 	TSharedRef<SRichTextBlock> MakeProjectSettingsDescriptionBox() const;
-	TSharedRef<SWidget> MakeProjectSettingsOptionsBox() const;
+	TSharedRef<SWidget> MakeProjectSettingsOptionsBox();
 
 private:
 
@@ -258,12 +259,11 @@ private:
 	/** The slate widget representing the list of templates */
 	TSharedPtr<STileView<TSharedPtr<FTemplateItem>>> TemplateListView;
 
+	TSharedPtr<SDecoratedEnumCombo<int32>> ProjectLanguageEnum;
+
 	/** Names for pages */
 	static FName TemplatePageName;
 	static FName NameAndLocationPageName;
 
-	TSharedPtr<SRichTextBlock> ProjectSettingsDescriptionBox;
-	TSharedPtr<SWidget> ProjectSettingsOptionsBox;
-
 	bool bProjectSettingsHidden;
 };
diff --git a/Engine/Source/Editor/HardwareTargeting/Public/SDecoratedEnumCombo.h b/Engine/Source/Editor/HardwareTargeting/Public/SDecoratedEnumCombo.h
--- a/Engine/Source/Editor/HardwareTargeting/Public/SDecoratedEnumCombo.h
+++ b/Engine/Source/Editor/HardwareTargeting/Public/SDecoratedEnumCombo.h
@@ -27,14 +27,14 @@ class SDecoratedEnumCombo : public SCompoundWidget
 public:
 	struct FComboOption
 	{
-		FComboOption(TEnumType InValue, FSlateIcon InIcon, FText InText, bool bInChoosable = true)
+		FComboOption(TEnumType InValue, FSlateIcon InIcon, FText InText, TAttribute<bool> bInChoosable = true)
 			: Value(InValue), Icon(InIcon), Text(InText), bChoosable(bInChoosable)
 		{}
 
 		TEnumType Value;
 		FSlateIcon Icon;
 		FText Text;
-		bool bChoosable;
+		TAttribute<bool> bChoosable;
 	};
 
 	DECLARE_DELEGATE_OneParam(FOnEnumChanged, TEnumType)
@@ -91,6 +91,8 @@ public:
 		];
 	}
 
+	TArray<FComboOption>& GetOptions() { return Options; }
+
 private:
 
 	TArray<FComboOption> Options;
@@ -152,7 +154,7 @@ private:
 		for (int32 Index = 0; Index < Options.Num(); ++Index)
 		{
 			const auto& Option = Options[Index];
-			if (!Option.bChoosable)
+			if (Option.bChoosable.Get() == false)
 			{
 				continue;
 			}
