Template category Icon is now treated as a relative path to an image. Removed Image field since it was unused.

Fixed category description being slightly clipped when wrapping.

Fixed a gap that was created when no major categories existed due to the TileView still reserving one row of space for itself.

diff --git a/Engine/Source/Editor/GameProjectGeneration/Classes/TemplateProjectDefs.h b/Engine/Source/Editor/GameProjectGeneration/Classes/TemplateProjectDefs.h
--- a/Engine/Source/Editor/GameProjectGeneration/Classes/TemplateProjectDefs.h
+++ b/Engine/Source/Editor/GameProjectGeneration/Classes/TemplateProjectDefs.h
@@ -184,12 +184,11 @@ struct FTemplateCategoryDef
 	UPROPERTY()
 	TArray<FLocalizedTemplateString> LocalizedDescriptions;
 
+	/** Reference to an icon to display for this category. Should be around 128x128. */
 	UPROPERTY()
-	FName Icon;
-
-	UPROPERTY()
-	FName Image;
+	FString Icon;
 
+	/** Is this a major top-level category? Major categories are displayed as full rows, eg. the Game category.*/
 	UPROPERTY()
 	bool IsMajorCategory;
 };
diff --git a/Engine/Source/Editor/GameProjectGeneration/Private/GameProjectGenerationModule.cpp b/Engine/Source/Editor/GameProjectGeneration/Private/GameProjectGenerationModule.cpp
--- a/Engine/Source/Editor/GameProjectGeneration/Private/GameProjectGenerationModule.cpp
+++ b/Engine/Source/Editor/GameProjectGeneration/Private/GameProjectGenerationModule.cpp
@@ -195,8 +195,10 @@ void FGameProjectGenerationModule::LoadTemplateCategories()
 				TemplateCategory->Key = Category.Key;
 				TemplateCategory->DisplayName = FLocalizedTemplateString::GetLocalizedText(Category.LocalizedDisplayNames);
 				TemplateCategory->Description = FLocalizedTemplateString::GetLocalizedText(Category.LocalizedDescriptions);
-				TemplateCategory->Icon = FEditorStyle::GetBrush(Category.Icon);
-				TemplateCategory->Image = FEditorStyle::GetBrush(Category.Image);
+				
+				const FName BrushName(*Category.Icon);
+				TemplateCategory->Icon = new FSlateDynamicImageBrush(BrushName, FVector2D(128, 128));
+
 				TemplateCategory->IsMajor = Category.IsMajorCategory;
 			}
 		}
diff --git a/Engine/Source/Editor/GameProjectGeneration/Private/GameProjectUtils.cpp b/Engine/Source/Editor/GameProjectGeneration/Private/GameProjectUtils.cpp
--- a/Engine/Source/Editor/GameProjectGeneration/Private/GameProjectUtils.cpp
+++ b/Engine/Source/Editor/GameProjectGeneration/Private/GameProjectUtils.cpp
@@ -1226,6 +1226,16 @@ UTemplateCategories* GameProjectUtils::LoadTemplateCategories(const FString& Roo
 	{
 		TemplateCategories = NewObject<UTemplateCategories>();
 		TemplateCategories->LoadConfig(UTemplateCategories::StaticClass(), *TemplateCategoriesIniFilename);
+
+		for (FTemplateCategoryDef& Category : TemplateCategories->Categories)
+		{
+			// attempt to resolve the icon relative to the root directory
+			FString TemplateCategoryIcon = RootDir / Category.Icon;
+			if (FPlatformFileManager::Get().GetPlatformFile().FileExists(*TemplateCategoryIcon))
+			{
+				Category.Icon = TemplateCategoryIcon;
+			}
+		}
 	}
 
 	return TemplateCategories;
diff --git a/Engine/Source/Editor/GameProjectGeneration/Private/SGameProjectDialog.cpp b/Engine/Source/Editor/GameProjectGeneration/Private/SGameProjectDialog.cpp
--- a/Engine/Source/Editor/GameProjectGeneration/Private/SGameProjectDialog.cpp
+++ b/Engine/Source/Editor/GameProjectGeneration/Private/SGameProjectDialog.cpp
@@ -246,6 +246,16 @@ TSharedRef<SWidget> SGameProjectDialog::CreateLandingPage()
 		]
 	];
 
+	if (MajorTemplateCategories.Num() == 0)
+	{
+		MajorCategoryTileView->SetVisibility(EVisibility::Collapsed);
+	}
+
+	if (MinorTemplateCategories.Num() == 0)
+	{
+		MinorCategoryTileView->SetVisibility(EVisibility::Collapsed);
+	}
+
 	return Widget;
 }
 
@@ -256,8 +266,6 @@ END_SLATE_FUNCTION_BUILD_OPTIMIZATION
 
 TSharedRef<ITableRow> SGameProjectDialog::ConstructTile(TSharedPtr<FTemplateCategory> Item, const TSharedRef<STableViewBase>& TableView)
 {
-	const float TextWidth = (Item->IsMajor ? GameProjectDialogDefs::TextWidth * 2 : GameProjectDialogDefs::TextWidth)/* - GameProjectDialogDefs::Padding * 2*/;
-
 	TSharedRef<STableRow<TSharedPtr<FTemplateCategory>>> Row = SNew(STableRow<TSharedPtr<FTemplateCategory>>, TableView)
 	.Style(FEditorStyle::Get(), "GameProjectDialog.TemplateListView.TableRow")
 	[
@@ -282,7 +290,7 @@ TSharedRef<ITableRow> SGameProjectDialog::ConstructTile(TSharedPtr<FTemplateCate
 				.Padding(0)
 				[
 					SNew(SImage)
-					.Image(Item->Image)
+					.Image(Item->Icon)
 				]
 			]
 
@@ -308,7 +316,7 @@ TSharedRef<ITableRow> SGameProjectDialog::ConstructTile(TSharedPtr<FTemplateCate
 				.Padding(0)
 				[
 					SNew(STextBlock)
-					.WrapTextAt(TextWidth)
+					.AutoWrapText(true)
 					.LineBreakPolicy(FBreakIterator::CreateWordBreakIterator())
 					.Justification(ETextJustify::Left)
 					.Text(Item->Description)
diff --git a/Engine/Source/Editor/GameProjectGeneration/Private/TemplateCategory.h b/Engine/Source/Editor/GameProjectGeneration/Private/TemplateCategory.h
--- a/Engine/Source/Editor/GameProjectGeneration/Private/TemplateCategory.h
+++ b/Engine/Source/Editor/GameProjectGeneration/Private/TemplateCategory.h
@@ -17,9 +17,6 @@ struct FTemplateCategory
 	/** A thumbnail to help identify this category (on the tab) */
 	const FSlateBrush* Icon;
 
-	/** A thumbnail to help identify this category (on the screenshot)*/
-	const FSlateBrush* Image;
-
 	/** A unique key for this category */
 	FName Key;
 
