Added default category that is visible when no TemplateCategories.ini has been synced.

New Project dialog no longer displays the Languages Available field in the preview window if the selected category contains no templates.

Added warning for cases when a template is referring to a category that doesn't exist.

diff --git a/Engine/Source/Editor/GameProjectGeneration/Classes/TemplateProjectDefs.h b/Engine/Source/Editor/GameProjectGeneration/Classes/TemplateProjectDefs.h
--- a/Engine/Source/Editor/GameProjectGeneration/Classes/TemplateProjectDefs.h
+++ b/Engine/Source/Editor/GameProjectGeneration/Classes/TemplateProjectDefs.h
@@ -123,6 +123,10 @@ public:
 	UPROPERTY(config)
 	bool bSkipProjectSettings;
 
+	/** Is this an enterprise template? */
+	UPROPERTY(config)
+	bool bIsEnterprise;
+
 	/* Optional list of settings to hide. If none are specified, then all settings are shown. */
 	UPROPERTY(config)
 	TArray<ETemplateSetting> HiddenSettings;
diff --git a/Engine/Source/Editor/GameProjectGeneration/Private/GameProjectUtils.cpp b/Engine/Source/Editor/GameProjectGeneration/Private/GameProjectUtils.cpp
--- a/Engine/Source/Editor/GameProjectGeneration/Private/GameProjectUtils.cpp
+++ b/Engine/Source/Editor/GameProjectGeneration/Private/GameProjectUtils.cpp
@@ -43,6 +43,7 @@
 #include "DefaultTemplateProjectDefs.h"
 #include "SNewClassDialog.h"
 #include "FeaturedClasses.inl"
+#include "TemplateCategory.h"
 
 #include "Features/IModularFeatures.h"
 
@@ -1265,8 +1266,25 @@ UTemplateProjectDefs* GameProjectUtils::LoadTemplateDefs(const FString& ProjectD
 				UE_LOG(LogGameProjectGeneration, Error, TEXT("Failed to find template project defs class '%s', using default."), *ClassName);
 			}
 		}
+
 		TemplateDefs = NewObject<UTemplateProjectDefs>(GetTransientPackage(), ClassToConstruct);
 		TemplateDefs->LoadConfig(UTemplateProjectDefs::StaticClass(), *TemplateDefsIniFilename);
+
+		TArray<TSharedPtr<FTemplateCategory>> AllTemplateCategories;
+		FGameProjectGenerationModule::Get().GetAllTemplateCategories(AllTemplateCategories);
+
+		for (const FName& CategoryKey : TemplateDefs->Categories)
+		{
+			bool bCategoryExists = AllTemplateCategories.ContainsByPredicate([&CategoryKey](const TSharedPtr<FTemplateCategory>& Category)
+				{
+					return Category->Key == CategoryKey;
+				});
+
+			if (!bCategoryExists)
+			{
+				UE_LOG(LogGameProjectGeneration, Warning, TEXT("Failed to find category definition named '%s', not defined in TemplateCategories.ini."), *CategoryKey.ToString());
+			}
+		}
 	}
 
 	return TemplateDefs;
diff --git a/Engine/Source/Editor/GameProjectGeneration/Private/SGameProjectDialog.cpp b/Engine/Source/Editor/GameProjectGeneration/Private/SGameProjectDialog.cpp
--- a/Engine/Source/Editor/GameProjectGeneration/Private/SGameProjectDialog.cpp
+++ b/Engine/Source/Editor/GameProjectGeneration/Private/SGameProjectDialog.cpp
@@ -126,6 +126,19 @@ TSharedRef<SWidget> SGameProjectDialog::CreateLandingPage()
 	TArray<TSharedPtr<FTemplateCategory>> AllTemplateCategories;
 	FGameProjectGenerationModule::Get().GetAllTemplateCategories(AllTemplateCategories);
 
+	if (AllTemplateCategories.Num() == 0)
+	{
+		TSharedPtr<FTemplateCategory> DefaultCategory = MakeShareable(new FTemplateCategory());
+		static const FName DefaultCategoryKey("Default");
+		DefaultCategory->Key = DefaultCategoryKey;
+		DefaultCategory->DisplayName = LOCTEXT("ProjectDialog_DefaultCategoryName", "Blank Project");
+		DefaultCategory->Description = LOCTEXT("ProjectDialog_DefaultCategoryDescription", "Create a new blank Unreal project.");
+		DefaultCategory->IsMajor = true;
+		DefaultCategory->Icon = FEditorStyle::GetBrush("GameProjectDialog.DefaultGameThumbnail");
+
+		AllTemplateCategories.Add(DefaultCategory);
+	}
+
 	for (const TSharedPtr<FTemplateCategory>& Category : AllTemplateCategories)
 	{
 		if (Category->IsMajor)
@@ -469,9 +482,6 @@ void SGameProjectDialog::OnFinishClicked()
 	{
 		NewProjectWizard->CreateAndOpenProject();
 	}
-
-	TSharedPtr<SWindow> Window = FSlateApplication::Get().FindWidgetWindow(AsShared());
-	Window->RequestDestroyWindow();
 }
 
 void SGameProjectDialog::OnCancelClicked() const
diff --git a/Engine/Source/Editor/GameProjectGeneration/Private/SNewProjectWizard.cpp b/Engine/Source/Editor/GameProjectGeneration/Private/SNewProjectWizard.cpp
--- a/Engine/Source/Editor/GameProjectGeneration/Private/SNewProjectWizard.cpp
+++ b/Engine/Source/Editor/GameProjectGeneration/Private/SNewProjectWizard.cpp
@@ -484,19 +484,23 @@ void SNewProjectWizard::Construct( const FArguments& InArgs )
 							.AutoHeight()
 							.Padding(FMargin(0.0f, 5.0f, 0.0f, 5.0f))
 							[
-								SNew(SVerticalBox)
-								+ SVerticalBox::Slot()
-								[
-									SNew(STextBlock)
-									.TextStyle(FEditorStyle::Get(), "GameProjectDialog.FeatureText")
-									.Text(LOCTEXT("LanguagesAvailable", "Languages Available:"))
-								]
-								+ SVerticalBox::Slot()
-								.AutoHeight()
+								SNew(SBox)
+								.Visibility(this, &SNewProjectWizard::GetSelectedTemplateLanguageVisibility)
 								[
-									SNew(STextBlock)
-									.AutoWrapText(true)
-									.Text(this, &SNewProjectWizard::GetSelectedTemplateLanguages)
+									SNew(SVerticalBox)
+									+ SVerticalBox::Slot()
+									[
+										SNew(STextBlock)
+										.TextStyle(FEditorStyle::Get(), "GameProjectDialog.FeatureText")
+										.Text(LOCTEXT("LanguagesAvailable", "Languages Available:"))
+									]
+									+ SVerticalBox::Slot()
+									.AutoHeight()
+									[
+										SNew(STextBlock)
+										.AutoWrapText(true)
+										.Text(this, &SNewProjectWizard::GetSelectedTemplateLanguages)
+									]
 								]
 							]
 						]
@@ -664,6 +668,17 @@ FText SNewProjectWizard::GetSelectedTemplateLanguages() const
 	return FText::GetEmpty();
 }
 
+EVisibility SNewProjectWizard::GetSelectedTemplateLanguageVisibility() const
+{
+	FText Languages = GetSelectedTemplateLanguages();
+	if (Languages.IsEmpty())
+	{
+		return EVisibility::Collapsed;
+	}
+
+	return EVisibility::Visible;
+}
+
 FText SNewProjectWizard::GetCurrentProjectFileName() const
 {
 	return FText::FromString( CurrentProjectFileName );
@@ -927,6 +942,7 @@ TMap<FName, TArray<TSharedPtr<FTemplateItem>> >& SNewProjectWizard::FindTemplate
 					Template->AssetTypes = TemplateDefs->AssetTypes;
 					Template->bSkipProjectSettings = TemplateDefs->bSkipProjectSettings;
 					Template->HiddenSettings = TemplateDefs->HiddenSettings;
+					Template->bIsEnterprise = TemplateDefs->bIsEnterprise;
 
 					Template->Name = TemplateDefs->GetDisplayNameText();
 					if (Template->Name.IsEmpty())
@@ -992,10 +1008,35 @@ TMap<FName, TArray<TSharedPtr<FTemplateItem>> >& SNewProjectWizard::FindTemplate
 		}
 	}
 
+	if (Templates.Num() == 0)
+	{
+		static const FName DefaultCategory("Default");
+		Templates.Add(DefaultCategory);
+	}
+
+	{
+		TSharedPtr<FTemplateItem> BlankTemplate = MakeShareable(new FTemplateItem());
+		BlankTemplate->Name = LOCTEXT("BlankProjectName", "Blank");
+		BlankTemplate->Description = LOCTEXT("BlankProjectDescription", "A clean empty project with no code.");
+		BlankTemplate->Key = TEXT("Blank");
+		BlankTemplate->SortKey = TEXT("_1");
+		BlankTemplate->Thumbnail = MakeShareable(new FSlateBrush(*FEditorStyle::GetBrush("GameProjectDialog.BlankProjectThumbnail")));
+		BlankTemplate->PreviewImage = MakeShareable(new FSlateBrush(*FEditorStyle::GetBrush("GameProjectDialog.BlankProjectPreview")));
+		BlankTemplate->BlueprintProjectFile = TEXT("");
+		BlankTemplate->CodeProjectFile = TEXT("");
+		BlankTemplate->bSkipProjectSettings = false;
+		BlankTemplate->bIsEnterprise = false;
+
+		for (auto& Pair : Templates)
+		{
+			Pair.Value.Add(BlankTemplate);
+		}
+	}
+
 	return Templates;
 }
 
-void SNewProjectWizard::SetDefaultProjectLocation( )
+void SNewProjectWizard::SetDefaultProjectLocation()
 {
 	FString DefaultProjectFilePath;
 	
@@ -1164,9 +1205,7 @@ bool SNewProjectWizard::CreateProject( const FString& ProjectFile )
 	ProjectInfo.DefaultGraphicsPerformance = SelectedGraphicsPreset;
 	ProjectInfo.bForceExtendedLuminanceRange = ProjectInfo.TemplateFile.IsEmpty();
 	ProjectInfo.bEnableVR = bEnableVR;
-
-	// @todo (sebastiann): This is a poor way to determine Enterprise-ness.
-	ProjectInfo.bIsEnterpriseProject = ActiveCategory != TEXT("Game");
+	ProjectInfo.bIsEnterpriseProject = SelectedTemplate->bIsEnterprise;
 
 	if (!GameProjectUtils::CreateProject(ProjectInfo, FailReason, FailLog))
 	{
diff --git a/Engine/Source/Editor/GameProjectGeneration/Private/SNewProjectWizard.h b/Engine/Source/Editor/GameProjectGeneration/Private/SNewProjectWizard.h
--- a/Engine/Source/Editor/GameProjectGeneration/Private/SNewProjectWizard.h
+++ b/Engine/Source/Editor/GameProjectGeneration/Private/SNewProjectWizard.h
@@ -102,8 +102,11 @@ private:
 	/** Get a string that details the class types referenced in the selected template */
 	FText GetSelectedTemplateClassTypes() const;
 
-	/** Get a string that details the class types referenced in the selected template */
+	/** Get a string that details the languages supported by the selected template */
 	FText GetSelectedTemplateLanguages() const;
+
+	/** Get visibility of the languages display. */
+	EVisibility GetSelectedTemplateLanguageVisibility() const;
 	
 	/** Get a visiblity of the class types display. If the string is empty this return Collapsed otherwise it will return Visible */
 	EVisibility GetSelectedTemplateClassVisibility() const;
diff --git a/Engine/Source/Editor/GameProjectGeneration/Private/TemplateCategory.h b/Engine/Source/Editor/GameProjectGeneration/Private/TemplateCategory.h
--- a/Engine/Source/Editor/GameProjectGeneration/Private/TemplateCategory.h
+++ b/Engine/Source/Editor/GameProjectGeneration/Private/TemplateCategory.h
@@ -22,4 +22,7 @@ struct FTemplateCategory
 
 	/** Is this a major or minor category? Controls where in the UI it will show up. */
 	bool IsMajor;
+
+	/** Is this an enterprise project? Will end up including Datasmith modules if true. */
+	bool IsEnterprise;
 };
diff --git a/Engine/Source/Editor/GameProjectGeneration/Private/TemplateItem.h b/Engine/Source/Editor/GameProjectGeneration/Private/TemplateItem.h
--- a/Engine/Source/Editor/GameProjectGeneration/Private/TemplateItem.h
+++ b/Engine/Source/Editor/GameProjectGeneration/Private/TemplateItem.h
@@ -27,4 +27,6 @@ struct FTemplateItem
 
 	bool		bSkipProjectSettings;
 	TArray<ETemplateSetting> HiddenSettings;
+
+	bool		bIsEnterprise;
 };
