Fixed check firing in SRecentProjectBrowser when no recent projects could be found. Most notably when syncing a fresh branch and launching with no project on the command-line.

diff --git a/Engine/Source/Editor/GameProjectGeneration/Private/SProjectBrowser.cpp b/Engine/Source/Editor/GameProjectGeneration/Private/SProjectBrowser.cpp
--- a/Engine/Source/Editor/GameProjectGeneration/Private/SProjectBrowser.cpp
+++ b/Engine/Source/Editor/GameProjectGeneration/Private/SProjectBrowser.cpp
@@ -1498,6 +1498,10 @@ void SRecentProjectBrowser::Construct(const FArguments& InArgs)
 
 	const FString EngineIdentifier = FDesktopPlatformModule::Get()->GetCurrentEngineIdentifier();
 
+	TSharedRef<FProjectCategory> NewCategory = MakeShareable(new FProjectCategory);
+	NewCategory->CategoryName = LOCTEXT("ProjectDialog_Recent", "Recent Projects");
+	ProjectCategories.Add(NewCategory);
+
 	TArray<FString> RecentProjects = GetDefault<UEditorSettings>()->RecentlyOpenedProjectFiles;
 	for (int32 Idx = 0; Idx < InArgs._NumProjects && Idx < RecentProjects.Num(); ++Idx)
 	{
@@ -1534,17 +1538,16 @@ void SRecentProjectBrowser::Construct(const FArguments& InArgs)
 	if (ProjectCategories.Num() > 0)
 	{
 		const TSharedRef<FProjectCategory> Category = ProjectCategories[0];
-		if (ensure(Category->ProjectItemsSource.Num() > 0) && ensure(Category->ProjectTileView.IsValid()))
+
+		if (ensure(Category->ProjectTileView.IsValid()) && Category->ProjectItemsSource.Num() > 0)
 		{
 			Category->ProjectTileView->SetSelection(Category->ProjectItemsSource[0], ESelectInfo::Direct);
 		}
 	}
 
 	bHasProjectFiles = false;
-	for (auto CategoryIt = ProjectCategories.CreateConstIterator(); CategoryIt; ++CategoryIt)
+	for (const TSharedRef<FProjectCategory>& Category : ProjectCategories)
 	{
-		const TSharedRef<FProjectCategory>& Category = *CategoryIt;
-
 		if (Category->ProjectItemsSource.Num() > 0)
 		{
 			bHasProjectFiles = true;
