Hidden settings are now respected when selecting another template, rather than only at construction time.

diff --git a/Engine/Source/Editor/GameProjectGeneration/Private/SGameProjectDialog.cpp b/Engine/Source/Editor/GameProjectGeneration/Private/SGameProjectDialog.cpp
--- a/Engine/Source/Editor/GameProjectGeneration/Private/SGameProjectDialog.cpp
+++ b/Engine/Source/Editor/GameProjectGeneration/Private/SGameProjectDialog.cpp
@@ -114,8 +114,9 @@ void SGameProjectDialog::Construct(const FArguments& InArgs, EMode InMode)
 			+ SWizard::Page()
 			.Name(GetPageTitle(GameProjectDialogDefs::ProjectSettingsPageIndex))
 			.CanShow(DialogMode != EMode::Open)
+			.OnEnter(this, &SGameProjectDialog::OnEnterSettingsPage)
 			[
-				NewProjectWizard->CreateProjectSettingsPage()
+				SAssignNew(ProjectSettingsPage, SBox)
 			]
 		]
 	];
@@ -431,6 +432,11 @@ int32 SGameProjectDialog::GetNextPageIndex(int32 Current) const
 	return INDEX_NONE;
 }
 
+void SGameProjectDialog::OnEnterSettingsPage()
+{
+	ProjectSettingsPage->SetContent(NewProjectWizard->CreateProjectSettingsPage());
+}
+
 bool SGameProjectDialog::OnCanFinish() const
 {
 	int32 CurrentPage = RootWizard->GetCurrentPageIndex();
diff --git a/Engine/Source/Editor/GameProjectGeneration/Private/SGameProjectDialog.h b/Engine/Source/Editor/GameProjectGeneration/Private/SGameProjectDialog.h
--- a/Engine/Source/Editor/GameProjectGeneration/Private/SGameProjectDialog.h
+++ b/Engine/Source/Editor/GameProjectGeneration/Private/SGameProjectDialog.h
@@ -50,6 +50,7 @@ private:
 	TSharedPtr<SProjectBrowser> ProjectBrowserPage;
 	TSharedPtr<SNewProjectWizard> NewProjectWizard;
 	TSharedPtr<SRecentProjectBrowser> RecentProjectBrowser;
+	TSharedPtr<SBox> ProjectSettingsPage;
 
 	TSharedPtr<STileView<TSharedPtr<FTemplateCategory>>> MajorCategoryTileView;
 	TSharedPtr<STileView<TSharedPtr<FTemplateCategory>>> MinorCategoryTileView;
@@ -78,6 +79,9 @@ private:
 	void OnCancelClicked() const;
 	bool OnCanFinish() const;
 	void OnFinishClicked();
+
+	void OnEnterSettingsPage();
+
 	FText GetFinishText() const;
 	FText GetFinishTooltip() const;
 	FText GetPageTitle(int32 Index) const;
diff --git a/Engine/Source/Editor/GameProjectGeneration/Private/SNewProjectWizard.cpp b/Engine/Source/Editor/GameProjectGeneration/Private/SNewProjectWizard.cpp
--- a/Engine/Source/Editor/GameProjectGeneration/Private/SNewProjectWizard.cpp
+++ b/Engine/Source/Editor/GameProjectGeneration/Private/SNewProjectWizard.cpp
@@ -573,25 +573,6 @@ void SNewProjectWizard::Tick( const FGeometry& AllottedGeometry, const double In
 void SNewProjectWizard::HandleTemplateListViewSelectionChanged(TSharedPtr<FTemplateItem> TemplateItem, ESelectInfo::Type SelectInfo)
 {
 	UpdateProjectFileValidity();
-
-	bool bIsBlueprintAvailable = !GetSelectedTemplateProperty(&FTemplateItem::BlueprintProjectFile).IsEmpty();
-	bool bIsCodeAvailable = !GetSelectedTemplateProperty(&FTemplateItem::CodeProjectFile).IsEmpty();
-
-	// if neither is available, then this is a blank template, so both are available
-	if (!bIsBlueprintAvailable && !bIsCodeAvailable)
-	{
-		bIsBlueprintAvailable = true;
-		bIsCodeAvailable = true;
-	}
-
-	TArray<SDecoratedEnumCombo<int32>::FComboOption>& LanguageOptions = ProjectLanguageEnum->GetOptions();
-	ensure(LanguageOptions[0].Text.IdenticalTo(LOCTEXT("ProjectDialog_Blueprint", "Blueprint")));
-	ensure(LanguageOptions[1].Text.IdenticalTo(LOCTEXT("ProjectDialog_Code", "C++")));
-
-	LanguageOptions[0].bChoosable = bIsBlueprintAvailable;
-	LanguageOptions[1].bChoosable = bIsCodeAvailable;
-
-	bShouldGenerateCode = !bIsBlueprintAvailable;
 }
 
 TSharedPtr<FTemplateItem> SNewProjectWizard::GetSelectedTemplateItem() const
@@ -1017,7 +998,7 @@ TMap<FName, TArray<TSharedPtr<FTemplateItem>> >& SNewProjectWizard::FindTemplate
 	{
 		TSharedPtr<FTemplateItem> BlankTemplate = MakeShareable(new FTemplateItem());
 		BlankTemplate->Name = LOCTEXT("BlankProjectName", "Blank");
-		BlankTemplate->Description = LOCTEXT("BlankProjectDescription", "A clean empty project with no code.");
+		BlankTemplate->Description = LOCTEXT("BlankProjectDescription", "A clean empty project with no code and default settings.");
 		BlankTemplate->Key = TEXT("Blank");
 		BlankTemplate->SortKey = TEXT("_1");
 		BlankTemplate->Thumbnail = MakeShareable(new FSlateBrush(*FEditorStyle::GetBrush("GameProjectDialog.BlankProjectThumbnail")));
@@ -1596,20 +1577,36 @@ TSharedRef<SWidget> SNewProjectWizard::MakeProjectSettingsOptionsBox()
 
 	if (!HiddenSettings.Contains(ETemplateSetting::Languages))
 	{
+		bool bIsBlueprintAvailable = !GetSelectedTemplateProperty(&FTemplateItem::BlueprintProjectFile).IsEmpty();
+		bool bIsCodeAvailable = !GetSelectedTemplateProperty(&FTemplateItem::CodeProjectFile).IsEmpty();
+
+		// if neither is available, then this is a blank template, so both are available
+		if (!bIsBlueprintAvailable && !bIsCodeAvailable)
+		{
+			bIsBlueprintAvailable = true;
+			bIsCodeAvailable = true;
+		}
+
+		bShouldGenerateCode = !bIsBlueprintAvailable;
+
 		TArray<SDecoratedEnumCombo<int32>::FComboOption> BlueprintOrCppOptions;
 		BlueprintOrCppOptions.Add(SDecoratedEnumCombo<int32>::FComboOption(
-			0, FSlateIcon(FEditorStyle::GetStyleSetName(), "GameProjectDialog.BlueprintImage_64"), LOCTEXT("ProjectDialog_Blueprint", "Blueprint")));
+			0, 
+			FSlateIcon(FEditorStyle::GetStyleSetName(), "GameProjectDialog.BlueprintImage_64"), 
+			LOCTEXT("ProjectDialog_Blueprint", "Blueprint"), 
+			bIsBlueprintAvailable));
 
 		BlueprintOrCppOptions.Add(SDecoratedEnumCombo<int32>::FComboOption(
-			1, FSlateIcon(FEditorStyle::GetStyleSetName(), "GameProjectDialog.CodeImage_64"), LOCTEXT("ProjectDialog_Code", "C++")));
+			1,
+			FSlateIcon(FEditorStyle::GetStyleSetName(), "GameProjectDialog.CodeImage_64"),
+			LOCTEXT("ProjectDialog_Code", "C++"), 
+			bIsCodeAvailable));
 
 		TSharedRef<SDecoratedEnumCombo<int32>> Enum = SNew(SDecoratedEnumCombo<int32>, MoveTemp(BlueprintOrCppOptions))
 			.SelectedEnum(this, &SNewProjectWizard::OnGetBlueprintOrCppIndex)
 			.OnEnumChanged(this, &SNewProjectWizard::OnSetBlueprintOrCppIndex)
 			.Orientation(Orient_Vertical);
 
-		ProjectLanguageEnum = Enum;
-
 		TSharedRef<SRichTextBlock> Description = SNew(SRichTextBlock)
 			.Text(LOCTEXT("ProjectDialog_BlueprintOrCppDescription", "Choose whether to create a Blueprint or C++ project."))
 			.AutoWrapText(true)
diff --git a/Engine/Source/Editor/GameProjectGeneration/Private/SNewProjectWizard.h b/Engine/Source/Editor/GameProjectGeneration/Private/SNewProjectWizard.h
--- a/Engine/Source/Editor/GameProjectGeneration/Private/SNewProjectWizard.h
+++ b/Engine/Source/Editor/GameProjectGeneration/Private/SNewProjectWizard.h
@@ -262,8 +262,6 @@ private:
 	/** The slate widget representing the list of templates */
 	TSharedPtr<STileView<TSharedPtr<FTemplateItem>>> TemplateListView;
 
-	TSharedPtr<SDecoratedEnumCombo<int32>> ProjectLanguageEnum;
-
 	/** Names for pages */
 	static FName TemplatePageName;
 	static FName NameAndLocationPageName;
