Only add the code-generated Blank template to categories in the New Project dialog that contain no other templates whatsoever.

diff --git a/Engine/Source/Editor/GameProjectGeneration/Private/SNewProjectWizard.cpp b/Engine/Source/Editor/GameProjectGeneration/Private/SNewProjectWizard.cpp
--- a/Engine/Source/Editor/GameProjectGeneration/Private/SNewProjectWizard.cpp
+++ b/Engine/Source/Editor/GameProjectGeneration/Private/SNewProjectWizard.cpp
@@ -995,6 +995,7 @@ TMap<FName, TArray<TSharedPtr<FTemplateItem>> >& SNewProjectWizard::FindTemplate
 		Templates.Add(DefaultCategory);
 	}
 
+	// add blank template to empty categories
 	{
 		TSharedPtr<FTemplateItem> BlankTemplate = MakeShareable(new FTemplateItem());
 		BlankTemplate->Name = LOCTEXT("BlankProjectName", "Blank");
@@ -1008,9 +1009,16 @@ TMap<FName, TArray<TSharedPtr<FTemplateItem>> >& SNewProjectWizard::FindTemplate
 		BlankTemplate->bSkipProjectSettings = false;
 		BlankTemplate->bIsEnterprise = false;
 
-		for (auto& Pair : Templates)
+		TArray<TSharedPtr<FTemplateCategory>> AllTemplateCategories;
+		FGameProjectGenerationModule::Get().GetAllTemplateCategories(AllTemplateCategories);
+
+		for (const TSharedPtr<FTemplateCategory>& Category : AllTemplateCategories)
 		{
-			Pair.Value.Add(BlankTemplate);
+			const TArray<TSharedPtr<FTemplateItem>>* CategoryEntry = Templates.Find(Category->Key);
+			if (CategoryEntry == nullptr)
+			{
+				Templates.Add(Category->Key).Add(BlankTemplate);
+			}
 		}
 	}
 
@@ -1573,7 +1581,7 @@ TSharedRef<SWidget> SNewProjectWizard::MakeProjectSettingsOptionsBox()
 		.FillColumn(1, 1.0f)
 		.FillColumn(3, 1.0f);
 
-	TArray<ETemplateSetting> HiddenSettings = GetSelectedTemplateProperty(&FTemplateItem::HiddenSettings);
+	const TArray<ETemplateSetting>& HiddenSettings = GetSelectedTemplateProperty(&FTemplateItem::HiddenSettings);
 
 	if (!HiddenSettings.Contains(ETemplateSetting::Languages))
 	{
