Added a Raytracing option to the New Project dialog settings.

Hooked up XR and Raytracing into the project creation phase.

Changed to use GameProjectDialog.* placeholder icons that are currently just some random AssetIcons.

Removed bSkipProjectSettings INI option, in favour of just having an enum value (ETemplateSetting::All) to control it.

diff --git a/Engine/Source/Editor/EditorStyle/Private/SlateEditorStyle.cpp b/Engine/Source/Editor/EditorStyle/Private/SlateEditorStyle.cpp
--- a/Engine/Source/Editor/EditorStyle/Private/SlateEditorStyle.cpp
+++ b/Engine/Source/Editor/EditorStyle/Private/SlateEditorStyle.cpp
@@ -1923,6 +1923,12 @@ void FSlateEditorStyle::FStyle::SetupGeneralStyles()
 
 		Set( "GameProjectDialog.IncludeStarterContent", new IMAGE_BRUSH( "/GameProjectDialog/IncludeStarterContent", FVector2D(64, 64) ) );
 		Set( "GameProjectDialog.NoStarterContent", new IMAGE_BRUSH( "/GameProjectDialog/NoStarterContent", FVector2D(64, 64) ) );
+		
+		Set("GameProjectDialog.XRDisabled", new IMAGE_BRUSH("Icons/AssetIcons/SkeletalMesh_64x", Icon64x64));
+		Set("GameProjectDialog.XREnabled", new IMAGE_BRUSH("Icons/AssetIcons/SkeletalMesh_64x", Icon64x64));
+
+		Set("GameProjectDialog.RaytracingDisabled", new IMAGE_BRUSH("Icons/AssetIcons/Procedural_Foliage_Spawner_64x", Icon64x64));
+		Set("GameProjectDialog.RaytracingEnabled", new IMAGE_BRUSH("Icons/AssetIcons/Procedural_Foliage_Spawner_64x", Icon64x64));
 
 		Set( "FilePath.FolderButton",
 			FButtonStyle(HoverHintOnly)
diff --git a/Engine/Source/Editor/GameProjectGeneration/Classes/TemplateProjectDefs.h b/Engine/Source/Editor/GameProjectGeneration/Classes/TemplateProjectDefs.h
--- a/Engine/Source/Editor/GameProjectGeneration/Classes/TemplateProjectDefs.h
+++ b/Engine/Source/Editor/GameProjectGeneration/Classes/TemplateProjectDefs.h
@@ -73,7 +73,9 @@ enum class ETemplateSetting
 	HardwareTarget,
 	GraphicsPreset,
 	StarterContent,
-	VR
+	XR,
+	Raytracing,
+	All
 };
 
 UCLASS(abstract,config=TemplateDefs,MinimalAPI)
@@ -115,14 +117,10 @@ public:
 	UPROPERTY(config)
 	FString AssetTypes;
 
-	/* Should we allow creation of a project from this template */
+	/* Should we allow creation of a project from this template. If this is false, the template is treated as a feature pack. */
 	UPROPERTY(config)
 	bool bAllowProjectCreation;
 
-	/* Should we skip the project settings page when creating a project from this template */
-	UPROPERTY(config)
-	bool bSkipProjectSettings;
-
 	/** Is this an enterprise template? */
 	UPROPERTY(config)
 	bool bIsEnterprise;
diff --git a/Engine/Source/Editor/GameProjectGeneration/Private/GameProjectUtils.cpp b/Engine/Source/Editor/GameProjectGeneration/Private/GameProjectUtils.cpp
--- a/Engine/Source/Editor/GameProjectGeneration/Private/GameProjectUtils.cpp
+++ b/Engine/Source/Editor/GameProjectGeneration/Private/GameProjectUtils.cpp
@@ -160,6 +160,50 @@ namespace
 
 		return MoveTemp(DefaultProjectSettings);
 	}
+
+	/** Set the state of XR plugins in OutProject based on the flags in InProjectInfo. */
+	void SetXRPluginStates(const FProjectInformation& InProjectInfo, FProjectDescriptor& OutProject)
+	{
+		static const FString XRPlugins[] = {
+			TEXT("MagicLeapMedia"),
+			TEXT("MagicLeap"),
+			TEXT("OculusVR"),
+			TEXT("SteamVR") };
+
+		if (!InProjectInfo.bEnableXR)
+		{
+			for (const FString& Plugin : XRPlugins)
+			{
+				int32 Index = OutProject.FindPluginReferenceIndex(Plugin);
+				if (Index == INDEX_NONE)
+				{
+					Index = OutProject.Plugins.AddDefaulted();
+					OutProject.Plugins[Index].Name = Plugin;
+				}
+
+				OutProject.Plugins[Index].bEnabled = false;
+			}
+		}
+	}
+
+	/** Get the configuration values for raytracing if enabled. */
+	FString GetRaytracingConfigString(const FProjectInformation& InProjectInfo)
+	{
+		FString RaytracingConfig;
+
+		if (InProjectInfo.bEnableRaytracing)
+		{
+			RaytracingConfig += LINE_TERMINATOR;
+			RaytracingConfig += TEXT("[/Script/WindowsTargetPlatform.WindowsTargetSettings]") LINE_TERMINATOR;
+			RaytracingConfig += TEXT("DefaultGraphicsRHI=DefaultGraphicsRHI_DX12") LINE_TERMINATOR;
+			RaytracingConfig += LINE_TERMINATOR;
+			RaytracingConfig += TEXT("[/Script/Engine.RendererSettings]") LINE_TERMINATOR;
+			RaytracingConfig += TEXT("r.SkinCache.CompileShaders=True") LINE_TERMINATOR;
+			RaytracingConfig += TEXT("r.RayTracing=True") LINE_TERMINATOR;
+		}
+
+		return MoveTemp(RaytracingConfig);
+	}
 } // namespace <>
 
 FText FNewClassInfo::GetClassName() const
@@ -1270,6 +1314,11 @@ UTemplateProjectDefs* GameProjectUtils::LoadTemplateDefs(const FString& ProjectD
 		TemplateDefs = NewObject<UTemplateProjectDefs>(GetTransientPackage(), ClassToConstruct);
 		TemplateDefs->LoadConfig(UTemplateProjectDefs::StaticClass(), *TemplateDefsIniFilename);
 
+		if (TemplateDefs->HiddenSettings.Num() > 1 && TemplateDefs->HiddenSettings.Contains(ETemplateSetting::All))
+		{
+			UE_LOG(LogGameProjectGeneration, Warning, TEXT("Template '%s' contains 'All' in HiddenSettings in addition to other entries. This is a mistake, and means that all settings will be hidden."), *ProjectDirectory);
+		}
+
 		TArray<TSharedPtr<FTemplateCategory>> AllTemplateCategories;
 		FGameProjectGenerationModule::Get().GetAllTemplateCategories(AllTemplateCategories);
 
@@ -1282,7 +1331,7 @@ UTemplateProjectDefs* GameProjectUtils::LoadTemplateDefs(const FString& ProjectD
 
 			if (!bCategoryExists)
 			{
-				UE_LOG(LogGameProjectGeneration, Warning, TEXT("Failed to find category definition named '%s', not defined in TemplateCategories.ini."), *CategoryKey.ToString());
+				UE_LOG(LogGameProjectGeneration, Warning, TEXT("Failed to find category definition named '%s' while loading template '%s', it is not defined in any TemplateCategories.ini."), *CategoryKey.ToString(), *ProjectDirectory);
 			}
 		}
 	}
@@ -1351,17 +1400,19 @@ bool GameProjectUtils::GenerateProjectFromScratch(const FProjectInformation& InP
 	// Generate the project file
 	{
 		// Set up the descriptor
-		FProjectDescriptor Descriptor;
+		FProjectDescriptor Project;
 		for(int32 Idx = 0; Idx < StartupModuleNames.Num(); Idx++)
 		{
-			Descriptor.Modules.Add(FModuleDescriptor(*StartupModuleNames[Idx]));
+			Project.Modules.Add(FModuleDescriptor(*StartupModuleNames[Idx]));
 		}
 
-		Descriptor.bIsEnterpriseProject = InProjectInfo.bIsEnterpriseProject;
+		Project.bIsEnterpriseProject = InProjectInfo.bIsEnterpriseProject;
+
+		SetXRPluginStates(InProjectInfo, Project);
 
 		// Try to save it
 		FText LocalFailReason;
-		if(!Descriptor.Save(InProjectInfo.ProjectFilename, LocalFailReason))
+		if(!Project.Save(InProjectInfo.ProjectFilename, LocalFailReason))
 		{
 			OutFailReason = LocalFailReason;
 			return false;
@@ -1587,6 +1638,8 @@ bool GameProjectUtils::CreateProjectFromTemplate(const FProjectInformation& InPr
 		FileContents += LINE_TERMINATOR;
 		FileContents += GetHardwareConfigString(InProjectInfo);
 
+		FileContents += GetRaytracingConfigString(InProjectInfo);
+
 		if ( !WriteOutputFile(DefaultEngineIniFilename, FileContents, OutFailReason) )
 		{
 			return false;
@@ -1752,6 +1805,9 @@ bool GameProjectUtils::CreateProjectFromTemplate(const FProjectInformation& InPr
 			ModuleInfo.Name = FName(*ModuleInfo.Name.ToString().Replace(*BaseSourceName, *BaseNewName));
 		}
 
+		// Disable XR plugins if desired
+		SetXRPluginStates(InProjectInfo, Project);
+
 		// Save it to disk
 		if(!Project.Save(InProjectInfo.ProjectFilename, OutFailReason))
 		{
@@ -1917,7 +1973,7 @@ FString GameProjectUtils::GetHardwareConfigString(const FProjectInformation& InP
 	HardwareTargeting += FString::Printf(TEXT("DefaultGraphicsPerformance=%s") LINE_TERMINATOR, *GraphicsPresetAsString);
 	HardwareTargeting += LINE_TERMINATOR;
 
-	return HardwareTargeting;
+	return MoveTemp(HardwareTargeting);
 }
 
 bool GameProjectUtils::GenerateConfigFiles(const FProjectInformation& InProjectInfo, TArray<FString>& OutCreatedFiles, FText& OutFailReason)
@@ -1981,6 +2037,9 @@ bool GameProjectUtils::GenerateConfigFiles(const FProjectInformation& InProjectI
 			}
 		}
 
+		FileContents += LINE_TERMINATOR;
+		FileContents += GetRaytracingConfigString(InProjectInfo);
+
 		if (WriteOutputFile(DefaultEngineIniFilename, FileContents, OutFailReason))
 		{
 			OutCreatedFiles.Add(DefaultEngineIniFilename);
@@ -4199,7 +4258,6 @@ bool GameProjectUtils::AddSharedContentToProject(const FProjectInformation &InPr
 			RequiredDetail = EFeaturePackDetailLevel::Standard;
 		}
 
-
 		TUniquePtr<FFeaturePackContentSource> TempFeaturePack = MakeUnique<FFeaturePackContentSource>();
 		bool bCopied = TempFeaturePack->InsertAdditionalResources(TemplateDefs->SharedContentPacks,RequiredDetail, DestFolder,CreatedFiles);
 		if( bCopied == false )
diff --git a/Engine/Source/Editor/GameProjectGeneration/Private/SNewProjectWizard.cpp b/Engine/Source/Editor/GameProjectGeneration/Private/SNewProjectWizard.cpp
--- a/Engine/Source/Editor/GameProjectGeneration/Private/SNewProjectWizard.cpp
+++ b/Engine/Source/Editor/GameProjectGeneration/Private/SNewProjectWizard.cpp
@@ -56,6 +56,7 @@ namespace NewProjectWizardDefs
 	const float ThumbnailSize = 64.f, ThumbnailPadding = 5.f;
 	const float ItemWidth = ThumbnailSize + 2*ThumbnailPadding;
 	const float ItemHeight = ItemWidth + 30;
+	const FName DefaultCategoryName = "Games";
 }
 
 /**
@@ -419,7 +420,7 @@ void SNewProjectWizard::Construct( const FArguments& InArgs )
 								SNew(STextBlock)
 								.AutoWrapText(true)
 								.TextStyle(FEditorStyle::Get(), "GameProjectDialog.FeatureText")
-								.Text(this, &SNewProjectWizard::GetSelectedTemplateProperty<FText>, &FTemplateItem::Name)
+								.Text(this, &SNewProjectWizard::GetSelectedTemplateProperty, &FTemplateItem::Name)
 							]
 						
 							// Template Description
@@ -427,7 +428,7 @@ void SNewProjectWizard::Construct( const FArguments& InArgs )
 							[
 								SNew(STextBlock)
 								.AutoWrapText(true)
-								.Text(this, &SNewProjectWizard::GetSelectedTemplateProperty<FText>, &FTemplateItem::Description)
+								.Text(this, &SNewProjectWizard::GetSelectedTemplateProperty, &FTemplateItem::Description)
 							]
 											
 							// Asset types
@@ -527,8 +528,8 @@ END_SLATE_FUNCTION_BUILD_OPTIMIZATION
 
 bool SNewProjectWizard::ShouldShowProjectSettingsPage() const
 {
-	bool bSkipSettings = GetSelectedTemplateProperty<bool>(&FTemplateItem::bSkipProjectSettings);
-	return !bSkipSettings;
+	const TArray<ETemplateSetting>& HiddenSettings = GetSelectedTemplateProperty(&FTemplateItem::HiddenSettings);
+	return !HiddenSettings.Contains(ETemplateSetting::All);
 }
 
 void SNewProjectWizard::OnSetCopyStarterContent(int32 InCopyStarterContent)
@@ -538,8 +539,8 @@ void SNewProjectWizard::OnSetCopyStarterContent(int32 InCopyStarterContent)
 
 EVisibility SNewProjectWizard::GetTemplateListLocationBoxVisibility() const
 {
-	bool bSkipSettings = GetSelectedTemplateProperty<bool>(&FTemplateItem::bSkipProjectSettings);
-	return bSkipSettings ? EVisibility::Visible : EVisibility::Collapsed;
+	bool bShowSettings = ShouldShowProjectSettingsPage();
+	return bShowSettings ? EVisibility::Collapsed : EVisibility::Visible;
 }
 
 EVisibility SNewProjectWizard::GetStarterContentWarningVisibility() const
@@ -573,6 +574,14 @@ void SNewProjectWizard::Tick( const FGeometry& AllottedGeometry, const double In
 void SNewProjectWizard::HandleTemplateListViewSelectionChanged(TSharedPtr<FTemplateItem> TemplateItem, ESelectInfo::Type SelectInfo)
 {
 	UpdateProjectFileValidity();
+
+	if (TemplateItem.IsValid())
+	{
+		if (TemplateItem->HiddenSettings.Contains(ETemplateSetting::StarterContent))
+		{
+			bCopyStarterContent = false;
+		}
+	}
 }
 
 TSharedPtr<FTemplateItem> SNewProjectWizard::GetSelectedTemplateItem() const
@@ -588,22 +597,22 @@ TSharedPtr<FTemplateItem> SNewProjectWizard::GetSelectedTemplateItem() const
 
 FText SNewProjectWizard::GetSelectedTemplateClassTypes() const
 {
-	return FText::FromString(GetSelectedTemplateProperty<FString>(&FTemplateItem::ClassTypes));
+	return FText::FromString(GetSelectedTemplateProperty(&FTemplateItem::ClassTypes));
 }
 
 EVisibility SNewProjectWizard::GetSelectedTemplateClassVisibility() const
 {
-	return GetSelectedTemplateProperty<FString>(&FTemplateItem::ClassTypes).IsEmpty() == false? EVisibility::Visible : EVisibility::Collapsed;
+	return GetSelectedTemplateProperty(&FTemplateItem::ClassTypes).IsEmpty() == false? EVisibility::Visible : EVisibility::Collapsed;
 }
 
 FText SNewProjectWizard::GetSelectedTemplateAssetTypes() const
 {
-	return FText::FromString(GetSelectedTemplateProperty<FString>(&FTemplateItem::AssetTypes));
+	return FText::FromString(GetSelectedTemplateProperty(&FTemplateItem::AssetTypes));
 }
 
 EVisibility SNewProjectWizard::GetSelectedTemplateAssetVisibility() const
 {
-	return GetSelectedTemplateProperty<FString>(&FTemplateItem::AssetTypes).IsEmpty() == false ? EVisibility::Visible : EVisibility::Collapsed;
+	return GetSelectedTemplateProperty(&FTemplateItem::AssetTypes).IsEmpty() == false ? EVisibility::Visible : EVisibility::Collapsed;
 }
 
 const FSlateBrush* SNewProjectWizard::GetSelectedTemplatePreviewImage() const
@@ -889,7 +898,7 @@ TMap<FName, TArray<TSharedPtr<FTemplateItem>> >& SNewProjectWizard::FindTemplate
 					TArray<FName> TemplateCategories = TemplateDefs->Categories;
 					if (TemplateCategories.Num() == 0)
 					{
-						TemplateCategories.Add("Game");
+						TemplateCategories.Add(NewProjectWizardDefs::DefaultCategoryName);
 					}
 
 					FString TemplateKey = Root;
@@ -921,7 +930,6 @@ TMap<FName, TArray<TSharedPtr<FTemplateItem>> >& SNewProjectWizard::FindTemplate
 					Template->Description = TemplateDefs->GetLocalizedDescription();
 					Template->ClassTypes = TemplateDefs->ClassTypes;
 					Template->AssetTypes = TemplateDefs->AssetTypes;
-					Template->bSkipProjectSettings = TemplateDefs->bSkipProjectSettings;
 					Template->HiddenSettings = TemplateDefs->HiddenSettings;
 					Template->bIsEnterprise = TemplateDefs->bIsEnterprise;
 
@@ -1006,7 +1014,6 @@ TMap<FName, TArray<TSharedPtr<FTemplateItem>> >& SNewProjectWizard::FindTemplate
 		BlankTemplate->PreviewImage = MakeShareable(new FSlateBrush(*FEditorStyle::GetBrush("GameProjectDialog.BlankProjectPreview")));
 		BlankTemplate->BlueprintProjectFile = TEXT("");
 		BlankTemplate->CodeProjectFile = TEXT("");
-		BlankTemplate->bSkipProjectSettings = false;
 		BlankTemplate->bIsEnterprise = false;
 
 		TArray<TSharedPtr<FTemplateCategory>> AllTemplateCategories;
@@ -1193,7 +1200,8 @@ bool SNewProjectWizard::CreateProject( const FString& ProjectFile )
 	ProjectInfo.TargetedHardware = SelectedHardwareClassTarget;
 	ProjectInfo.DefaultGraphicsPerformance = SelectedGraphicsPreset;
 	ProjectInfo.bForceExtendedLuminanceRange = ProjectInfo.TemplateFile.IsEmpty();
-	ProjectInfo.bEnableVR = bEnableVR;
+	ProjectInfo.bEnableXR = bEnableXR;
+	ProjectInfo.bEnableRaytracing = bEnableRaytracing;
 	ProjectInfo.bIsEnterpriseProject = SelectedTemplate->bIsEnterprise;
 
 	if (!GameProjectUtils::CreateProject(ProjectInfo, FailReason, FailLog))
@@ -1631,7 +1639,7 @@ TSharedRef<SWidget> SNewProjectWizard::MakeProjectSettingsOptionsBox()
 			Orient_Vertical);
 		
 		TSharedRef<SRichTextBlock> Description = SNew(SRichTextBlock)
-			.Text(LOCTEXT("HardwareClassTargetDescription", "Choose the closest equivalent target platform. Don't worry, you can change this later in the <RichTextBlock.BoldHighlight>Target Hardware</> section of <RichTextBlock.BoldHighlight>Project Settings</>."))
+			.Text(LOCTEXT("ProjectDialog_HardwareClassTargetDescription", "Choose the closest equivalent target platform. Don't worry, you can change this later in the <RichTextBlock.BoldHighlight>Target Hardware</> section of <RichTextBlock.BoldHighlight>Project Settings</>."))
 			.AutoWrapText(true)
 			.DecoratorStyleSet(&FEditorStyle::Get());
 
@@ -1646,7 +1654,7 @@ TSharedRef<SWidget> SNewProjectWizard::MakeProjectSettingsOptionsBox()
 			Orient_Vertical);
 
 		TSharedRef<SRichTextBlock> Description = SNew(SRichTextBlock)
-			.Text(LOCTEXT("GraphicsPresetDescription", "Choose the performance characteristics of your project."))
+			.Text(LOCTEXT("ProjectDialog_GraphicsPresetDescription", "Choose the performance characteristics of your project."))
 			.AutoWrapText(true)
 			.DecoratorStyleSet(&FEditorStyle::Get());
 
@@ -1657,12 +1665,17 @@ TSharedRef<SWidget> SNewProjectWizard::MakeProjectSettingsOptionsBox()
 	{
 		TArray<SDecoratedEnumCombo<int32>::FComboOption> StarterContentOptions;
 		StarterContentOptions.Add(SDecoratedEnumCombo<int32>::FComboOption(
-			0, FSlateIcon(FEditorStyle::GetStyleSetName(), "GameProjectDialog.NoStarterContent"), LOCTEXT("NoStarterContent", "No Starter Content")));
+			0,
+			FSlateIcon(FEditorStyle::GetStyleSetName(), "GameProjectDialog.NoStarterContent"),
+			LOCTEXT("NoStarterContent", "No Starter Content")));
 
 		// Only add the option to add starter content if its there to add !
 		bool bIsStarterAvailable = GameProjectUtils::IsStarterContentAvailableForNewProjects();
 		StarterContentOptions.Add(SDecoratedEnumCombo<int32>::FComboOption(
-			1, FSlateIcon(FEditorStyle::GetStyleSetName(), "GameProjectDialog.IncludeStarterContent"), LOCTEXT("IncludeStarterContent", "With Starter Content"), bIsStarterAvailable));
+			1,
+			FSlateIcon(FEditorStyle::GetStyleSetName(), "GameProjectDialog.IncludeStarterContent"),
+			LOCTEXT("IncludeStarterContent", "With Starter Content"),
+			bIsStarterAvailable));
 
 		TSharedRef<SWidget> Enum = SNew(SOverlay)
 			+ SOverlay::Slot()
@@ -1693,23 +1706,50 @@ TSharedRef<SWidget> SNewProjectWizard::MakeProjectSettingsOptionsBox()
 		AddToProjectSettingsGrid(GridPanel, Enum, Description, CurrentSlot);
 	}
 
-	if (!HiddenSettings.Contains(ETemplateSetting::VR))
+	if (!HiddenSettings.Contains(ETemplateSetting::XR))
 	{
-		TArray<SDecoratedEnumCombo<int32>::FComboOption> VREnabledOptions;
-		VREnabledOptions.Add(SDecoratedEnumCombo<int32>::FComboOption(
-			0, FSlateIcon(FEditorStyle::GetStyleSetName(), "LevelEditor.ToggleVR"/*"GameProjectDialog.VRDisabled"*/), LOCTEXT("VRDisabled", "VR Disabled")));
+		TArray<SDecoratedEnumCombo<int32>::FComboOption> VirtualRealityOptions;
+		VirtualRealityOptions.Add(SDecoratedEnumCombo<int32>::FComboOption(
+			0, FSlateIcon(FEditorStyle::GetStyleSetName(), "GameProjectDialog.XRDisabled"),
+			LOCTEXT("XRDisabled", "XR Disabled")));
+
+		VirtualRealityOptions.Add(SDecoratedEnumCombo<int32>::FComboOption(
+			1, 
+			FSlateIcon(FEditorStyle::GetStyleSetName(), "GameProjectDialog.XREnabled"),
+			LOCTEXT("XREnabled", "XR Enabled")));
+
+		TSharedRef<SDecoratedEnumCombo<int32>> Enum = SNew(SDecoratedEnumCombo<int32>, MoveTemp(VirtualRealityOptions))
+			.SelectedEnum(this, &SNewProjectWizard::OnGetXREnabled)
+			.OnEnumChanged(this, &SNewProjectWizard::OnSetXREnabled)
+			.Orientation(Orient_Vertical);
 
-		VREnabledOptions.Add(SDecoratedEnumCombo<int32>::FComboOption(
-			1, FSlateIcon(FEditorStyle::GetStyleSetName(), "LevelEditor.ToggleVR"/*"GameProjectDialog.VREnabled"*/), LOCTEXT("VREnabled", "VR Enabled")));
+		TSharedRef<SRichTextBlock> Description = SNew(SRichTextBlock)
+			.Text(LOCTEXT("ProjectDialog_XREnabledDescription", "Choose if XR should be enabled in the new project."))
+			.AutoWrapText(true)
+			.DecoratorStyleSet(&FEditorStyle::Get());
 
-		TSharedRef<SDecoratedEnumCombo<int32>> Enum = SNew(SDecoratedEnumCombo<int32>, MoveTemp(VREnabledOptions))
-			.SelectedEnum(this, &SNewProjectWizard::OnGetVREnabled)
-			.OnEnumChanged(this, &SNewProjectWizard::OnSetVREnabled)
-			.Orientation(Orient_Vertical)
-			.ToolTipText(LOCTEXT("EnableVR_Tooltip", "Enable to to add VR support to the created project."));
+		AddToProjectSettingsGrid(GridPanel, Enum, Description, CurrentSlot);
+	}
+
+	if (!HiddenSettings.Contains(ETemplateSetting::Raytracing))
+	{
+		TArray<SDecoratedEnumCombo<int32>::FComboOption> RaytracingOptions;
+		RaytracingOptions.Add(SDecoratedEnumCombo<int32>::FComboOption(
+			0, FSlateIcon(FEditorStyle::GetStyleSetName(), "GameProjectDialog.RaytracingDisabled"),
+			LOCTEXT("ProjectDialog_RaytracingDisabled", "Raytracing Disabled")));
+
+		RaytracingOptions.Add(SDecoratedEnumCombo<int32>::FComboOption(
+			1,
+			FSlateIcon(FEditorStyle::GetStyleSetName(), "GameProjectDialog.RaytracingEnabled"),
+			LOCTEXT("ProjectDialog_RaytracingEnabled", "Raytracing Enabled")));
+
+		TSharedRef<SDecoratedEnumCombo<int32>> Enum = SNew(SDecoratedEnumCombo<int32>, MoveTemp(RaytracingOptions))
+			.SelectedEnum(this, &SNewProjectWizard::OnGetRaytracingEnabled)
+			.OnEnumChanged(this, &SNewProjectWizard::OnSetRaytracingEnabled)
+			.Orientation(Orient_Vertical);
 
 		TSharedRef<SRichTextBlock> Description = SNew(SRichTextBlock)
-			.Text(LOCTEXT("VREnabledDescription", "Choose if VR should be enabled in the new project."))
+			.Text(LOCTEXT("ProjectDialog_RaytracingDescription", "Choose if real-time raytracing should be enabled in the new project."))
 			.AutoWrapText(true)
 			.DecoratorStyleSet(&FEditorStyle::Get());
 
diff --git a/Engine/Source/Editor/GameProjectGeneration/Private/SNewProjectWizard.h b/Engine/Source/Editor/GameProjectGeneration/Private/SNewProjectWizard.h
--- a/Engine/Source/Editor/GameProjectGeneration/Private/SNewProjectWizard.h
+++ b/Engine/Source/Editor/GameProjectGeneration/Private/SNewProjectWizard.h
@@ -156,8 +156,11 @@ private:
 	/** Returns the tooltip text (actual warning) of the starter content warning */
 	FText GetStarterContentWarningTooltip() const;
 
-	int32 OnGetVREnabled() const { return bEnableVR ? 1 : 0; }
-	void OnSetVREnabled(int32 InVREnabledIndex) { bEnableVR = InVREnabledIndex == 1; }
+	int32 OnGetXREnabled() const { return bEnableXR ? 1 : 0; }
+	void OnSetXREnabled(int32 InEnumIndex) { bEnableXR = InEnumIndex == 1; }
+
+	int32 OnGetRaytracingEnabled() const { return bEnableRaytracing ? 1 : 0; }
+	void OnSetRaytracingEnabled(int32 InEnumIndex) { bEnableRaytracing = InEnumIndex == 1; }
 
 	int32 OnGetBlueprintOrCppIndex() const;
 	void OnSetBlueprintOrCppIndex(int32 Index);
@@ -244,8 +247,11 @@ private:
 	/** True if user has selected to copy starter content. */
 	bool bCopyStarterContent;
 
-	/** Whether or not to enable VR in the created project. */
-	bool bEnableVR;
+	/** Whether or not to enable XR in the created project. */
+	bool bEnableXR;
+
+	/** Whether or not to enable Raytracing in the created project. */
+	bool bEnableRaytracing;
 
 	/** Whether or not we should use the blueprint or C++ version of this template. */
 	bool bShouldGenerateCode;
diff --git a/Engine/Source/Editor/GameProjectGeneration/Private/TemplateItem.h b/Engine/Source/Editor/GameProjectGeneration/Private/TemplateItem.h
--- a/Engine/Source/Editor/GameProjectGeneration/Private/TemplateItem.h
+++ b/Engine/Source/Editor/GameProjectGeneration/Private/TemplateItem.h
@@ -25,7 +25,6 @@ struct FTemplateItem
 	FString		CodeProjectFile;
 	FString		BlueprintProjectFile;
 
-	bool		bSkipProjectSettings;
 	TArray<ETemplateSetting> HiddenSettings;
 
 	bool		bIsEnterprise;
diff --git a/Engine/Source/Editor/GameProjectGeneration/Public/GameProjectUtils.h b/Engine/Source/Editor/GameProjectGeneration/Public/GameProjectUtils.h
--- a/Engine/Source/Editor/GameProjectGeneration/Public/GameProjectUtils.h
+++ b/Engine/Source/Editor/GameProjectGeneration/Public/GameProjectUtils.h
@@ -25,8 +25,9 @@ struct FProjectInformation
 	bool bCopyStarterContent = false;
 	
 	bool bForceExtendedLuminanceRange = false; // See "r.DefaultFeature.AutoExposure.ExtendDefaultLuminanceRange"
-	bool bEnableVR = false;
+	bool bEnableXR = false;
 	bool bIsEnterpriseProject = false;
+	bool bEnableRaytracing = false;
 
 	EHardwareClass::Type TargetedHardware = EHardwareClass::Desktop;
 	EGraphicsPreset::Type DefaultGraphicsPerformance = EGraphicsPreset::Maximum;
