Fixed periodic validity checks not firing anymore and displaying an error when creating/removing project files in disk with the New Project dialog wizard open.

Removed extra sentence from description that was considered clutter.

Removed languages listing from templates because enterprise did not want it.

"Categories" label is hidden in New Project mode.

diff --git a/Engine/Source/Editor/GameProjectGeneration/Private/SGameProjectDialog.cpp b/Engine/Source/Editor/GameProjectGeneration/Private/SGameProjectDialog.cpp
--- a/Engine/Source/Editor/GameProjectGeneration/Private/SGameProjectDialog.cpp
+++ b/Engine/Source/Editor/GameProjectGeneration/Private/SGameProjectDialog.cpp
@@ -226,8 +226,8 @@ TSharedRef<SWidget> SGameProjectDialog::CreateLandingPage()
 					SNew(STextBlock)
 					.TextStyle(FEditorStyle::Get(), "GameProjectDialog.ProjectNamePathLabels")
 					.Text(LOCTEXT("ProjectDialog_Categories", "Categories"))
+					.Visibility(DialogMode == EMode::Both ? EVisibility::Visible : EVisibility::Collapsed)
 				]
-
 				+ SVerticalBox::Slot()
 				.AutoHeight()
 				.Padding(8, 8, 8, 0)
diff --git a/Engine/Source/Editor/GameProjectGeneration/Private/SNewProjectWizard.cpp b/Engine/Source/Editor/GameProjectGeneration/Private/SNewProjectWizard.cpp
--- a/Engine/Source/Editor/GameProjectGeneration/Private/SNewProjectWizard.cpp
+++ b/Engine/Source/Editor/GameProjectGeneration/Private/SNewProjectWizard.cpp
@@ -305,11 +305,8 @@ void SNewProjectWizard::Construct( const FArguments& InArgs )
 {
 	bProjectSettingsHidden = false;
 
-	LastValidityCheckTime = 0;
-	ValidityCheckFrequency = 4;
 	bLastGlobalValidityCheckSuccessful = true;
 	bLastNameAndLocationValidityCheckSuccessful = true;
-	bPreventPeriodicValidityChecksUntilNextChange = false;
 	bCopyStarterContent = GEditor ? GetDefault<UEditorSettings>()->bCopyStarterContentPreference : true;
 
 	SelectedHardwareClassTarget = EHardwareClass::Desktop;
@@ -479,31 +476,6 @@ void SNewProjectWizard::Construct( const FArguments& InArgs )
 									]
 								]
 							]
-
-							// Languages available
-							+ SVerticalBox::Slot()
-							.AutoHeight()
-							.Padding(FMargin(0.0f, 5.0f, 0.0f, 5.0f))
-							[
-								SNew(SBox)
-								.Visibility(this, &SNewProjectWizard::GetSelectedTemplateLanguageVisibility)
-								[
-									SNew(SVerticalBox)
-									+ SVerticalBox::Slot()
-									[
-										SNew(STextBlock)
-										.TextStyle(FEditorStyle::Get(), "GameProjectDialog.FeatureText")
-										.Text(LOCTEXT("LanguagesAvailable", "Languages Available:"))
-									]
-									+ SVerticalBox::Slot()
-									.AutoHeight()
-									[
-										SNew(STextBlock)
-										.AutoWrapText(true)
-										.Text(this, &SNewProjectWizard::GetSelectedTemplateLanguages)
-									]
-								]
-							]
 						]
 					]
 				]
@@ -560,17 +532,6 @@ FText SNewProjectWizard::GetStarterContentWarningTooltip() const
 	}
 }
 
-void SNewProjectWizard::Tick( const FGeometry& AllottedGeometry, const double InCurrentTime, const float InDeltaTime )
-{
-	// Every few seconds, the project file path is checked for validity in case the disk contents changed and the location is now valid or invalid.
-	// After project creation, periodic checks are disabled to prevent a brief message indicating that the project you created already exists.
-	// This feature is re-enabled if the user did not restart and began editing parameters again.
-	if ( !bPreventPeriodicValidityChecksUntilNextChange && (InCurrentTime > LastValidityCheckTime + ValidityCheckFrequency) )
-	{
-		UpdateProjectFileValidity();
-	}
-}
-
 void SNewProjectWizard::HandleTemplateListViewSelectionChanged(TSharedPtr<FTemplateItem> TemplateItem, ESelectInfo::Type SelectInfo)
 {
 	UpdateProjectFileValidity();
@@ -627,48 +588,6 @@ EVisibility SNewProjectWizard::GetSelectedTemplatePreviewVisibility() const
 	return PreviewImage.IsValid() ? EVisibility::Visible : EVisibility::Collapsed;
 }
 
-FText SNewProjectWizard::GetSelectedTemplateLanguages() const
-{
-	TSharedPtr<FTemplateItem> SelectedItem = GetSelectedTemplateItem();
-	if (SelectedItem.IsValid())
-	{
-		bool bHasBlueprint = !SelectedItem->BlueprintProjectFile.IsEmpty();
-		bool bHasCode = !SelectedItem->CodeProjectFile.IsEmpty();
-
-		if (!bHasBlueprint && !bHasCode)
-		{
-			bHasBlueprint = true;
-			bHasCode = true;
-		}
-
-		if (bHasBlueprint && bHasCode)
-		{
-			return LOCTEXT("ProjectDialog_BlueprintAndCode", "Blueprint, C++");
-		}
-		else if (bHasBlueprint)
-		{
-			return LOCTEXT("ProjectDialog_Blueprint", "Blueprint");
-		}
-		else if (bHasCode)
-		{
-			return LOCTEXT("ProjectDialog_Code", "C++");
-		}
-	}
-
-	return FText::GetEmpty();
-}
-
-EVisibility SNewProjectWizard::GetSelectedTemplateLanguageVisibility() const
-{
-	FText Languages = GetSelectedTemplateLanguages();
-	if (Languages.IsEmpty())
-	{
-		return EVisibility::Collapsed;
-	}
-
-	return EVisibility::Visible;
-}
-
 FText SNewProjectWizard::GetCurrentProjectFileName() const
 {
 	return FText::FromString( CurrentProjectFileName );
@@ -801,7 +720,6 @@ FText SNewProjectWizard::GetGlobalErrorLabelText() const
 	return FText::GetEmpty();
 }
 
-
 FReply SNewProjectWizard::OnCloseGlobalErrorLabelClicked()
 {
 	PersistentGlobalErrorLabelText = FText();
@@ -809,13 +727,11 @@ FReply SNewProjectWizard::OnCloseGlobalErrorLabelClicked()
 	return FReply::Handled();
 }
 
-
 EVisibility SNewProjectWizard::GetNameAndLocationErrorLabelVisibility() const
 {
 	return GetNameAndLocationErrorLabelText().IsEmpty() ? EVisibility::Collapsed : EVisibility::Visible;
 }
 
-
 FText SNewProjectWizard::GetNameAndLocationErrorLabelText() const
 {
 	if ( !bLastNameAndLocationValidityCheckSuccessful )
@@ -1096,7 +1012,7 @@ void SNewProjectWizard::SetDefaultProjectLocation()
 	}
 }
 
-void SNewProjectWizard::UpdateProjectFileValidity( )
+void SNewProjectWizard::UpdateProjectFileValidity()
 {
 	// Global validity
 	{
@@ -1160,11 +1076,6 @@ void SNewProjectWizard::UpdateProjectFileValidity( )
 			}
 		}
 	}
-
-	LastValidityCheckTime = FSlateApplication::Get().GetCurrentTime();
-
-	// Since this function was invoked, periodic validity checks should be re-enabled if they were disabled.
-	bPreventPeriodicValidityChecksUntilNextChange = false;
 }
 
 bool SNewProjectWizard::IsCompilerRequired() const
@@ -1241,9 +1152,6 @@ void SNewProjectWizard::CreateAndOpenProject( )
 		return;
 	}
 
-	// Prevent periodic validity checks. This is to prevent a brief error message about the project already existing while you are exiting.
-	bPreventPeriodicValidityChecksUntilNextChange = true;
-
 	if (bShouldGenerateCode)
 	{
 	    // If the engine is installed it is already compiled, so we can try to build and open a new project immediately. Non-installed situations might require building
@@ -1352,7 +1260,7 @@ void SNewProjectWizard::SetCurrentCategory(FName Category)
 
 TSharedRef<SWidget> SNewProjectWizard::MakeProjectLocationWidget()
 {
-	return SNew(SVerticalBox)
+	TSharedRef<SWidget> Widget = SNew(SVerticalBox)
 	+ SVerticalBox::Slot()
 	.HAlign(HAlign_Center)
 	.Padding(8)
@@ -1383,6 +1291,15 @@ TSharedRef<SWidget> SNewProjectWizard::MakeProjectLocationWidget()
 			.OnNameChanged(this, &SNewProjectWizard::OnCurrentProjectFileNameChanged)
 		]
 	];
+
+	Widget->RegisterActiveTimer(1.0f, FWidgetActiveTimerDelegate::CreateLambda(
+		[this](double, float)
+		{
+			UpdateProjectFileValidity();
+			return EActiveTimerReturnType::Continue;
+		}));
+
+	return Widget;
 }
 
 TSharedRef<SWidget> SNewProjectWizard::CreateProjectSettingsPage()
@@ -1699,7 +1616,7 @@ TSharedRef<SWidget> SNewProjectWizard::MakeProjectSettingsOptionsBox()
 
 		TSharedRef<SRichTextBlock> Description =
 			SNew(SRichTextBlock)
-			.Text(LOCTEXT("CopyStarterContent_ToolTip", "Enable to include an additional content pack containing simple placeable meshes with basic materials and textures.\nYou can opt out of including this to create a project that only has the bare essentials for the selected project template.\nYou can also add the <RichTextBlock.BoldHighlight>Starter Content</> to your project later using <RichTextBlock.BoldHighlight>Content Browser</>."))
+			.Text(LOCTEXT("CopyStarterContent_ToolTip", "Enable to include an additional content pack containing simple placeable meshes with basic materials and textures.\nYou can also add the <RichTextBlock.BoldHighlight>Starter Content</> to your project later using <RichTextBlock.BoldHighlight>Content Browser</>."))
 			.AutoWrapText(true)
 			.DecoratorStyleSet(&FEditorStyle::Get());
 		
diff --git a/Engine/Source/Editor/GameProjectGeneration/Private/SNewProjectWizard.h b/Engine/Source/Editor/GameProjectGeneration/Private/SNewProjectWizard.h
--- a/Engine/Source/Editor/GameProjectGeneration/Private/SNewProjectWizard.h
+++ b/Engine/Source/Editor/GameProjectGeneration/Private/SNewProjectWizard.h
@@ -33,8 +33,6 @@ public:
 	/** Constructs this widget with InArgs */
 	void Construct( const FArguments& InArgs );
 
-	virtual void Tick( const FGeometry& AllottedGeometry, const double InCurrentTime, const float InDeltaTime ) override;
-
 	/** Populates TemplateItemsSource with templates found on disk */
 	TMap<FName, TArray<TSharedPtr<FTemplateItem>> >& FindTemplateProjects();
 
@@ -102,12 +100,6 @@ private:
 	/** Get a string that details the class types referenced in the selected template */
 	FText GetSelectedTemplateClassTypes() const;
 
-	/** Get a string that details the languages supported by the selected template */
-	FText GetSelectedTemplateLanguages() const;
-
-	/** Get visibility of the languages display. */
-	EVisibility GetSelectedTemplateLanguageVisibility() const;
-	
 	/** Get a visiblity of the class types display. If the string is empty this return Collapsed otherwise it will return Visible */
 	EVisibility GetSelectedTemplateClassVisibility() const;
 	
@@ -223,15 +215,6 @@ private:
 
 	FSimpleDelegate OnTemplateDoubleClick;
 
-	/** The last time that the selected project file path was checked for validity. This is used to throttle I/O requests to a reasonable frequency */
-	double LastValidityCheckTime;
-
-	/** The frequency in seconds for validity checks while the dialog is idle. Changes to the path immediately update the validity. */
-	double ValidityCheckFrequency;
-
-	/** Periodic checks for validity will not occur while this flag is true. Used to prevent a frame of "this project already exists" while exiting after a successful creation. */
-	bool bPreventPeriodicValidityChecksUntilNextChange;
-
 	/** The global error text from the last validity check */
 	FText LastGlobalValidityErrorText;
 
