The New Project dialog no longer shows a second redundant Open button.

diff --git a/Engine/Source/Editor/GameProjectGeneration/Private/SGameProjectDialog.cpp b/Engine/Source/Editor/GameProjectGeneration/Private/SGameProjectDialog.cpp
--- a/Engine/Source/Editor/GameProjectGeneration/Private/SGameProjectDialog.cpp
+++ b/Engine/Source/Editor/GameProjectGeneration/Private/SGameProjectDialog.cpp
@@ -63,7 +63,8 @@ void SGameProjectDialog::Construct(const FArguments& InArgs, EMode InMode)
 
 	NewProjectWizard = SNew(SNewProjectWizard)
 		.OnTemplateDoubleClick(this, &SGameProjectDialog::OnTemplateDoubleClick);
-	ProjectBrowserPage = SNew(SProjectBrowser);
+	ProjectBrowserPage = SNew(SProjectBrowser)
+		.HideOpenButton(true);
 
 	const float UniformPadding = 16.0f;
 
diff --git a/Engine/Source/Editor/GameProjectGeneration/Private/SProjectBrowser.cpp b/Engine/Source/Editor/GameProjectGeneration/Private/SProjectBrowser.cpp
--- a/Engine/Source/Editor/GameProjectGeneration/Private/SProjectBrowser.cpp
+++ b/Engine/Source/Editor/GameProjectGeneration/Private/SProjectBrowser.cpp
@@ -324,7 +324,7 @@ void SProjectBrowser::Construct( const FArguments& InArgs )
 					SNew(SButton)
 					.Text(LOCTEXT("BrowseProjectButton", "Browse..."))
 					.OnClicked(this, &SProjectBrowser::OnBrowseToProjectClicked)
-					.ContentPadding( FCoreStyle::Get().GetMargin("StandardDialog.ContentPadding") )
+					.ContentPadding(FCoreStyle::Get().GetMargin("StandardDialog.ContentPadding"))
 				]
 
 				// Open Button
@@ -335,7 +335,8 @@ void SProjectBrowser::Construct( const FArguments& InArgs )
 					.Text(LOCTEXT("OpenProjectButton", "Open"))
 					.OnClicked(this, &SProjectBrowser::HandleOpenProjectButtonClicked)
 					.IsEnabled(this, &SProjectBrowser::HandleOpenProjectButtonIsEnabled)
-					.ContentPadding( FCoreStyle::Get().GetMargin("StandardDialog.ContentPadding") )
+					.ContentPadding(FCoreStyle::Get().GetMargin("StandardDialog.ContentPadding"))
+					.Visibility(InArgs._HideOpenButton ? EVisibility::Collapsed : EVisibility::Visible)
 				]
 			]
 		]
diff --git a/Engine/Source/Editor/GameProjectGeneration/Private/SProjectBrowser.h b/Engine/Source/Editor/GameProjectGeneration/Private/SProjectBrowser.h
--- a/Engine/Source/Editor/GameProjectGeneration/Private/SProjectBrowser.h
+++ b/Engine/Source/Editor/GameProjectGeneration/Private/SProjectBrowser.h
@@ -31,9 +31,13 @@ public:
 
 	DECLARE_DELEGATE(FNewProjectScreenRequested)
 
-	SLATE_BEGIN_ARGS(SProjectBrowser) { }
+	SLATE_BEGIN_ARGS(SProjectBrowser) :
+		_HideOpenButton(false)
+	{}
 
-	SLATE_ARGUMENT(FProjectSelectionChanged, OnSelectionChanged);
+		SLATE_ARGUMENT(FProjectSelectionChanged, OnSelectionChanged);
+
+		SLATE_ARGUMENT(bool, HideOpenButton)
 
 	SLATE_END_ARGS()
 
