Creating a Blank project with starter content enabled now opens the Minimal_Default map.

Added bIsBlank to TemplateDefs.ini, because it is impossible to differentiate between a "blank" project where we want to use the Starter Content map and one where we don't without the added context.

diff --git a/Engine/Source/Editor/GameProjectGeneration/Classes/TemplateProjectDefs.h b/Engine/Source/Editor/GameProjectGeneration/Classes/TemplateProjectDefs.h
--- a/Engine/Source/Editor/GameProjectGeneration/Classes/TemplateProjectDefs.h
+++ b/Engine/Source/Editor/GameProjectGeneration/Classes/TemplateProjectDefs.h
@@ -125,6 +125,10 @@ public:
 	UPROPERTY(config)
 	bool bIsEnterprise;
 
+	/** Is this a blank template? Determines whether we can override the default map when creating the project. */
+	UPROPERTY(config)
+	bool bIsBlank;
+
 	/* Optional list of settings to hide. If none are specified, then all settings are shown. */
 	UPROPERTY(config)
 	TArray<ETemplateSetting> HiddenSettings;
diff --git a/Engine/Source/Editor/GameProjectGeneration/Private/GameProjectUtils.cpp b/Engine/Source/Editor/GameProjectGeneration/Private/GameProjectUtils.cpp
--- a/Engine/Source/Editor/GameProjectGeneration/Private/GameProjectUtils.cpp
+++ b/Engine/Source/Editor/GameProjectGeneration/Private/GameProjectUtils.cpp
@@ -93,8 +93,6 @@
 TWeakPtr<SNotificationItem> GameProjectUtils::UpdateGameProjectNotification = NULL;
 TWeakPtr<SNotificationItem> GameProjectUtils::WarningProjectNameNotification = NULL;
 
-FString GameProjectUtils::DefaultFeaturePackExtension(TEXT(".upack"));
-
 bool GameProjectUtils::bUseAudioMixerForAllPlatforms = false;
 
 struct FAudioDefaultPlatformSettings
@@ -202,6 +200,32 @@ namespace
 
 		return MoveTemp(RaytracingConfig);
 	}
+
+	FString GetDefaultMapConfigString(const FProjectInformation& InProjectInfo)
+	{
+		FString DefaultMapConfig;
+
+		if (InProjectInfo.bIsBlankTemplate &&
+			InProjectInfo.bCopyStarterContent &&
+			GameProjectUtils::IsStarterContentAvailableForNewProjects())
+		{
+			DefaultMapConfig += LINE_TERMINATOR;
+			DefaultMapConfig += TEXT("[/Script/EngineSettings.GameMapsSettings]") LINE_TERMINATOR;
+
+			if (InProjectInfo.TargetedHardware == EHardwareClass::Mobile)
+			{
+				DefaultMapConfig += TEXT("EditorStartupMap=/Game/MobileStarterContent/Maps/Minimal_Default") LINE_TERMINATOR;
+				DefaultMapConfig += TEXT("GameDefaultMap=/Game/MobileStarterContent/Maps/Minimal_Default") LINE_TERMINATOR;
+			}
+			else
+			{
+				DefaultMapConfig += TEXT("EditorStartupMap=/Game/StarterContent/Maps/Minimal_Default") LINE_TERMINATOR;
+				DefaultMapConfig += TEXT("GameDefaultMap=/Game/StarterContent/Maps/Minimal_Default") LINE_TERMINATOR;
+			}
+		}
+
+		return MoveTemp(DefaultMapConfig);
+	}
 } // namespace <>
 
 FText FNewClassInfo::GetClassName() const
@@ -795,13 +819,11 @@ bool GameProjectUtils::OpenCodeIDE(const FString& ProjectFile, FText& OutFailRea
 	return true;
 }
 
-void GameProjectUtils::GetStarterContentFiles(TArray<FString>& OutFilenames)
+bool GameProjectUtils::IsStarterContentAvailableForNewProjects()
 {
-	FString const SrcFolder = FPaths::FeaturePackDir();
-
-	FString SearchPath = TEXT("*");
-	SearchPath += DefaultFeaturePackExtension;
-	IFileManager::Get().FindFilesRecursive(OutFilenames, *SrcFolder, *SearchPath, /*Files=*/true, /*Directories=*/false);
+	TArray<FString> OutFilenames;
+	IFileManager::Get().FindFilesRecursive(OutFilenames, *FPaths::FeaturePackDir(), TEXT("*StarterContent.upack"), /*Files=*/true, /*Directories=*/false);
+	return OutFilenames.Num() > 0;
 }
 
 bool GameProjectUtils::CreateProject(const FProjectInformation& InProjectInfo, FText& OutFailReason, FText& OutFailLog, TArray<FString>* OutCreatedFiles)
@@ -1632,7 +1654,7 @@ bool GameProjectUtils::CreateProjectFromTemplate(const FProjectInformation& InPr
 
 	const FString ProjectConfigPath = DestFolder / TEXT("Config");
 
-	// Write out the hardware class target settings chosen for this project
+	// Write out DefaultEngine.ini settings
 	{
 		const FString DefaultEngineIniFilename = ProjectConfigPath / TEXT("DefaultEngine.ini");
 
@@ -1640,12 +1662,13 @@ bool GameProjectUtils::CreateProjectFromTemplate(const FProjectInformation& InPr
 		// Load the existing file - if it doesn't exist we create it
 		FFileHelper::LoadFileToString(FileContents, *DefaultEngineIniFilename);
 
-		FileContents += LINE_TERMINATOR;
 		FileContents += GetHardwareConfigString(InProjectInfo);
 
 		FileContents += GetRaytracingConfigString(InProjectInfo);
 
-		if ( !WriteOutputFile(DefaultEngineIniFilename, FileContents, OutFailReason) )
+		FileContents += GetDefaultMapConfigString(InProjectInfo);
+
+		if (!WriteOutputFile(DefaultEngineIniFilename, FileContents, OutFailReason))
 		{
 			return false;
 		}
@@ -1973,10 +1996,10 @@ FString GameProjectUtils::GetHardwareConfigString(const FProjectInformation& InP
 	FString GraphicsPresetAsString;
 	UEnum::GetValueAsString(TEXT("/Script/HardwareTargeting.EGraphicsPreset"), InProjectInfo.DefaultGraphicsPerformance, /*out*/ GraphicsPresetAsString);
 
+	HardwareTargeting += LINE_TERMINATOR;
 	HardwareTargeting += TEXT("[/Script/HardwareTargeting.HardwareTargetingSettings]") LINE_TERMINATOR;
 	HardwareTargeting += FString::Printf(TEXT("TargetedHardwareClass=%s") LINE_TERMINATOR, *TargetHardwareAsString);
 	HardwareTargeting += FString::Printf(TEXT("DefaultGraphicsPerformance=%s") LINE_TERMINATOR, *GraphicsPresetAsString);
-	HardwareTargeting += LINE_TERMINATOR;
 
 	return MoveTemp(HardwareTargeting);
 }
@@ -1993,58 +2016,49 @@ bool GameProjectUtils::GenerateConfigFiles(const FProjectInformation& InProjectI
 		const FString DefaultEngineIniFilename = ProjectConfigPath / TEXT("DefaultEngine.ini");
 		FString FileContents;
 
-		FileContents += TEXT("[URL]") LINE_TERMINATOR;
-
 		FileContents += GetHardwareConfigString(InProjectInfo);
-		FileContents += LINE_TERMINATOR;
+
+		FileContents += GetRaytracingConfigString(InProjectInfo);
 
 		if(bUseAudioMixerForAllPlatforms)
 		{
+			FileContents += LINE_TERMINATOR;
 			FileContents += TEXT("[Audio]") LINE_TERMINATOR;
 			FileContents += TEXT("UseAudioMixer=True") LINE_TERMINATOR;
-			FileContents += LINE_TERMINATOR;
 		}
 
 		if (InProjectInfo.bForceExtendedLuminanceRange)
 		{
+			FileContents += LINE_TERMINATOR;
 			FileContents += TEXT("[/Script/Engine.RendererSettings]") LINE_TERMINATOR;
 			FileContents += TEXT("r.DefaultFeature.AutoExposure.ExtendDefaultLuminanceRange=True") LINE_TERMINATOR;
-			FileContents += LINE_TERMINATOR;
 		}
 
 		if (InProjectInfo.bCopyStarterContent)
 		{
-			FString SpecificEditorStartupMap;
-			FString SpecificGameDefaultMap;
+			FileContents += LINE_TERMINATOR;
+			FileContents += TEXT("[/Script/EngineSettings.GameMapsSettings]") LINE_TERMINATOR;
 
-			// If we have starter content packs available, specify starter map
-			if( IsStarterContentAvailableForNewProjects() == true )
+			if (GameProjectUtils::IsStarterContentAvailableForNewProjects())
 			{
 				if (InProjectInfo.TargetedHardware == EHardwareClass::Mobile)
 				{
-					SpecificEditorStartupMap = TEXT("/Game/MobileStarterContent/Maps/Minimal_Default");
-					SpecificGameDefaultMap = TEXT("/Game/MobileStarterContent/Maps/Minimal_Default");
+					FileContents += TEXT("EditorStartupMap=/Game/MobileStarterContent/Maps/Minimal_Default") LINE_TERMINATOR;
+					FileContents += TEXT("GameDefaultMap=/Game/MobileStarterContent/Maps/Minimal_Default") LINE_TERMINATOR;
 				}
 				else
 				{
-					SpecificEditorStartupMap = TEXT("/Game/StarterContent/Maps/Minimal_Default");
-					SpecificGameDefaultMap = TEXT("/Game/StarterContent/Maps/Minimal_Default");
+					FileContents += TEXT("EditorStartupMap=/Game/StarterContent/Maps/Minimal_Default") LINE_TERMINATOR;
+					FileContents += TEXT("GameDefaultMap=/Game/StarterContent/Maps/Minimal_Default") LINE_TERMINATOR;
 				}
 			}
 
-			// Write out the settings for startup map and game default map
-			FileContents += TEXT("[/Script/EngineSettings.GameMapsSettings]") LINE_TERMINATOR;
-			FileContents += FString::Printf(TEXT("EditorStartupMap=%s") LINE_TERMINATOR, *SpecificEditorStartupMap);
-			FileContents += FString::Printf(TEXT("GameDefaultMap=%s") LINE_TERMINATOR, *SpecificGameDefaultMap);
 			if (InProjectInfo.bShouldGenerateCode)
 			{
 				FileContents += FString::Printf(TEXT("GlobalDefaultGameMode=\"/Script/%s.%sGameMode\"") LINE_TERMINATOR, *NewProjectName, *NewProjectName);
 			}
 		}
 
-		FileContents += LINE_TERMINATOR;
-		FileContents += GetRaytracingConfigString(InProjectInfo);
-
 		if (WriteOutputFile(DefaultEngineIniFilename, FileContents, OutFailReason))
 		{
 			OutCreatedFiles.Add(DefaultEngineIniFilename);
@@ -2364,15 +2378,6 @@ bool GameProjectUtils::GenerateCodeProjectFiles(const FString& ProjectFilename,
 	return true;
 }
 
-bool GameProjectUtils::IsStarterContentAvailableForNewProjects()
-{
-	TArray<FString> StarterContentFiles;
-	GetStarterContentFiles(StarterContentFiles);
-
-	bool bHasStaterContent = StarterContentFiles.FindByPredicate([&](const FString& Str){ return Str.Contains("StarterContent"); }) != nullptr;
-	return bHasStaterContent;
-}
-
 int32 FindSpecificBuildFiles(const TCHAR* Path, TMap<FString, FString>& BuildFiles)
 {
 	class FBuildFileVistor : public IPlatformFile::FDirectoryVisitor
@@ -4095,11 +4100,11 @@ bool GameProjectUtils::InsertFeaturePacksIntoINIFile(const FProjectInformation&
 		FString StarterPack;
 		if (InProjectInfo.TargetedHardware == EHardwareClass::Mobile)
 		{
-			StarterPack = TEXT("InsertPack=(PackSource=\"MobileStarterContent") + DefaultFeaturePackExtension + TEXT("\",PackName=\"StarterContent\")");
+			StarterPack = TEXT("InsertPack=(PackSource=\"MobileStarterContent.upack\",PackName=\"StarterContent\")");
 		}
 		else
 		{
-			StarterPack = TEXT("InsertPack=(PackSource=\"StarterContent")  + DefaultFeaturePackExtension + TEXT("\",PackName=\"StarterContent\")");
+			StarterPack = TEXT("InsertPack=(PackSource=\"StarterContent.upack\",PackName=\"StarterContent\")");
 		}
 		PackList.Add(StarterPack);
 	}
diff --git a/Engine/Source/Editor/GameProjectGeneration/Private/SNewProjectWizard.cpp b/Engine/Source/Editor/GameProjectGeneration/Private/SNewProjectWizard.cpp
--- a/Engine/Source/Editor/GameProjectGeneration/Private/SNewProjectWizard.cpp
+++ b/Engine/Source/Editor/GameProjectGeneration/Private/SNewProjectWizard.cpp
@@ -318,13 +318,13 @@ void SNewProjectWizard::Construct( const FArguments& InArgs )
 	FindTemplateProjects();
 	SetDefaultProjectLocation();
 
-	TemplateListView = SNew(STileView< TSharedPtr<FTemplateItem> >)
+	TemplateListView = SNew(STileView<TSharedPtr<FTemplateItem>>)
 	.ListItemsSource(&FilteredTemplateList)
 	.SelectionMode(ESelectionMode::Single)
 	.ClearSelectionOnClick(false)
 	.OnGenerateTile_Static(&STemplateTile::BuildTile)
-	.ItemHeight( NewProjectWizardDefs::ItemHeight )
-	.ItemWidth( NewProjectWizardDefs::ItemWidth )
+	.ItemHeight(NewProjectWizardDefs::ItemHeight)
+	.ItemWidth(NewProjectWizardDefs::ItemWidth)
 	.OnMouseButtonDoubleClick(this, &SNewProjectWizard::HandleTemplateListViewDoubleClick)
 	.OnSelectionChanged(this, &SNewProjectWizard::HandleTemplateListViewSelectionChanged);
 
@@ -848,6 +848,7 @@ TMap<FName, TArray<TSharedPtr<FTemplateItem>> >& SNewProjectWizard::FindTemplate
 					Template->AssetTypes = TemplateDefs->AssetTypes;
 					Template->HiddenSettings = TemplateDefs->HiddenSettings;
 					Template->bIsEnterprise = TemplateDefs->bIsEnterprise;
+					Template->bIsBlankTemplate = TemplateDefs->bIsBlank;
 
 					Template->Name = TemplateDefs->GetDisplayNameText();
 					if (Template->Name.IsEmpty())
@@ -931,6 +932,7 @@ TMap<FName, TArray<TSharedPtr<FTemplateItem>> >& SNewProjectWizard::FindTemplate
 		BlankTemplate->BlueprintProjectFile = TEXT("");
 		BlankTemplate->CodeProjectFile = TEXT("");
 		BlankTemplate->bIsEnterprise = false;
+		BlankTemplate->bIsBlankTemplate = true;
 
 		TArray<TSharedPtr<FTemplateCategory>> AllTemplateCategories;
 		FGameProjectGenerationModule::Get().GetAllTemplateCategories(AllTemplateCategories);
@@ -1089,7 +1091,7 @@ bool SNewProjectWizard::IsCompilerRequired() const
 	return false;
 }
 
-bool SNewProjectWizard::CreateProject( const FString& ProjectFile )
+bool SNewProjectWizard::CreateProject(const FString& ProjectFile)
 {
 	// Get the selected template
 	TSharedPtr<FTemplateItem> SelectedTemplate = GetSelectedTemplateItem();
@@ -1114,6 +1116,7 @@ bool SNewProjectWizard::CreateProject( const FString& ProjectFile )
 	ProjectInfo.bEnableXR = bEnableXR;
 	ProjectInfo.bEnableRaytracing = bEnableRaytracing;
 	ProjectInfo.bIsEnterpriseProject = SelectedTemplate->bIsEnterprise;
+	ProjectInfo.bIsBlankTemplate = SelectedTemplate->bIsBlankTemplate;
 
 	if (!GameProjectUtils::CreateProject(ProjectInfo, FailReason, FailLog))
 	{
@@ -1139,15 +1142,15 @@ bool SNewProjectWizard::CreateProject( const FString& ProjectFile )
 	return true;
 }
 
-void SNewProjectWizard::CreateAndOpenProject( )
+void SNewProjectWizard::CreateAndOpenProject()
 {
-	if( !CanCreateProject() )
+	if (!CanCreateProject())
 	{
 		return;
 	}
 
 	FString ProjectFile = GetProjectFilenameWithPath();
-	if ( !CreateProject(ProjectFile) )
+	if (!CreateProject(ProjectFile))
 	{
 		return;
 	}
@@ -1160,8 +1163,8 @@ void SNewProjectWizard::CreateAndOpenProject( )
 		{
 			if (GameProjectUtils::BuildCodeProject(ProjectFile))
 			{
-				OpenCodeIDE( ProjectFile );
-				OpenProject( ProjectFile );
+				OpenCodeIDE(ProjectFile);
+				OpenProject(ProjectFile);
 			}
 			else
 			{
@@ -1170,19 +1173,19 @@ void SNewProjectWizard::CreateAndOpenProject( )
 		}
 		else
 		{
-			OpenCodeIDE( ProjectFile );
+			OpenCodeIDE(ProjectFile);
 		}
 	}
 	else
 	{
-		OpenProject( ProjectFile );
+		OpenProject(ProjectFile );
 	}
 }
 
-bool SNewProjectWizard::OpenProject( const FString& ProjectFile )
+bool SNewProjectWizard::OpenProject(const FString& ProjectFile)
 {
 	FText FailReason;
-	if ( GameProjectUtils::OpenProject( ProjectFile, FailReason ) )
+	if (GameProjectUtils::OpenProject(ProjectFile, FailReason))
 	{
 		// Successfully opened the project, the editor is closing.
 		// Close this window in case something prevents the editor from closing (save dialog, quit confirmation, etc)
@@ -1194,11 +1197,11 @@ bool SNewProjectWizard::OpenProject( const FString& ProjectFile )
 	return false;
 }
 
-bool SNewProjectWizard::OpenCodeIDE( const FString& ProjectFile )
+bool SNewProjectWizard::OpenCodeIDE(const FString& ProjectFile)
 {
 	FText FailReason;
 
-	if ( GameProjectUtils::OpenCodeIDE( ProjectFile, FailReason ) )
+	if (GameProjectUtils::OpenCodeIDE(ProjectFile, FailReason))
 	{
 		// Successfully opened code editing IDE, the editor is closing
 		// Close this window in case something prevents the editor from closing (save dialog, quit confirmation, etc)
@@ -1206,7 +1209,7 @@ bool SNewProjectWizard::OpenCodeIDE( const FString& ProjectFile )
 		return true;
 	}
 
-	DisplayError( FailReason );
+	DisplayError(FailReason);
 	return false;
 }
 
diff --git a/Engine/Source/Editor/GameProjectGeneration/Private/TemplateItem.h b/Engine/Source/Editor/GameProjectGeneration/Private/TemplateItem.h
--- a/Engine/Source/Editor/GameProjectGeneration/Private/TemplateItem.h
+++ b/Engine/Source/Editor/GameProjectGeneration/Private/TemplateItem.h
@@ -28,4 +28,5 @@ struct FTemplateItem
 	TArray<ETemplateSetting> HiddenSettings;
 
 	bool		bIsEnterprise;
+	bool		bIsBlankTemplate;
 };
diff --git a/Engine/Source/Editor/GameProjectGeneration/Public/GameProjectUtils.h b/Engine/Source/Editor/GameProjectGeneration/Public/GameProjectUtils.h
--- a/Engine/Source/Editor/GameProjectGeneration/Public/GameProjectUtils.h
+++ b/Engine/Source/Editor/GameProjectGeneration/Public/GameProjectUtils.h
@@ -23,6 +23,7 @@ struct FProjectInformation
 
 	bool bShouldGenerateCode = false;
 	bool bCopyStarterContent = false;
+	bool bIsBlankTemplate = false;
 	
 	bool bForceExtendedLuminanceRange = false; // See "r.DefaultFeature.AutoExposure.ExtendDefaultLuminanceRange"
 	bool bEnableXR = false;
@@ -301,9 +302,6 @@ private:
 	 */
 	static bool AddSharedContentToProject(const FProjectInformation &InProjectInfo, TArray<FString> &CreatedFiles, FText& OutFailReason);
 
-	/** Returns list of starter content files */
-	static void GetStarterContentFiles(TArray<FString>& OutFilenames);
-
 	/** Returns the template defs ini filename */
 	static FString GetTemplateDefsFilename();
 
@@ -508,7 +506,6 @@ private:
 
 	static TWeakPtr<SNotificationItem> UpdateGameProjectNotification;
 	static TWeakPtr<SNotificationItem> WarningProjectNameNotification;
-	static FString DefaultFeaturePackExtension;
 
 	// Whether we should use AudioMixer for all platforms:
 	static bool bUseAudioMixerForAllPlatforms;
diff --git a/Templates/TP_Blank/Config/TemplateDefs.ini b/Templates/TP_Blank/Config/TemplateDefs.ini
--- a/Templates/TP_Blank/Config/TemplateDefs.ini
+++ b/Templates/TP_Blank/Config/TemplateDefs.ini
@@ -20,6 +20,9 @@ SortKey="_1"
 ; These should contain files only of concern to the template itself and not desired in the user's project.
 FoldersToIgnore=Media
 
+; Is this a blank template? Blank templates do not specify a default map.
+bIsBlank=True
+
 ; The files to ignore copying, relative to the template root folder. 
 ; %TEMPLATENAME% will be replaced by the template's folder name when resolving these. 
 ; Filename comparisons are case-insensitive. 
diff --git a/Templates/TP_BlankBP/Config/TemplateDefs.ini b/Templates/TP_BlankBP/Config/TemplateDefs.ini
--- a/Templates/TP_BlankBP/Config/TemplateDefs.ini
+++ b/Templates/TP_BlankBP/Config/TemplateDefs.ini
@@ -2,6 +2,8 @@
 
 SortKey="_1"
 
+bIsBlank=True
+
 Categories=Games
 
 LocalizedDisplayNames=(Language="en", Text="Blank")
