New Project wizard now correctly lists a Blank project when the Templates folder has not been synced.

Descriptions have been removed from the New Project wizard pages. First page has been renamed to "Select Template Category".

diff --git a/Engine/Source/Editor/GameProjectGeneration/Private/SGameProjectDialog.cpp b/Engine/Source/Editor/GameProjectGeneration/Private/SGameProjectDialog.cpp
--- a/Engine/Source/Editor/GameProjectGeneration/Private/SGameProjectDialog.cpp
+++ b/Engine/Source/Editor/GameProjectGeneration/Private/SGameProjectDialog.cpp
@@ -50,7 +50,7 @@ namespace GameProjectDialogDefs
 	constexpr int32 ProjectSettingsPageIndex = 3;
 	
 	static const FText LandingPageTitle = LOCTEXT("ProjectDialog_SelectProject", "Select Project");
-	static const FText SelectCategoryTitle = LOCTEXT("ProjectDialog_SelectCategory", "Select Category");
+	static const FText SelectCategoryTitle = LOCTEXT("ProjectDialog_SelectTemplateCategory", "Select Template Category");
 	static const FText ProjectBrowserTitle = LOCTEXT("ProjectDialog_ProjectBrowser", "Project Browser");
 	static const FText TemplateListTitle = LOCTEXT("ProjectDialog_SelectTemplate", "Select Template");
 	static const FText ProjectSettingsTitle = LOCTEXT("ProjectDialog_ProjectSettings", "Project Settings");
@@ -123,7 +123,7 @@ void SGameProjectDialog::Construct(const FArguments& InArgs, EMode InMode)
 	];
 }
 
-TSharedRef<SWidget> SGameProjectDialog::CreateLandingPage()
+void SGameProjectDialog::GetAllTemplateCategories(TArray<TSharedPtr<FTemplateCategory>>& OutCategories)
 {
 	TArray<TSharedPtr<FTemplateCategory>> AllTemplateCategories;
 	FGameProjectGenerationModule::Get().GetAllTemplateCategories(AllTemplateCategories);
@@ -141,6 +141,14 @@ TSharedRef<SWidget> SGameProjectDialog::CreateLandingPage()
 		AllTemplateCategories.Add(DefaultCategory);
 	}
 
+	OutCategories = AllTemplateCategories;
+}
+
+TSharedRef<SWidget> SGameProjectDialog::CreateLandingPage()
+{
+	TArray<TSharedPtr<FTemplateCategory>> AllTemplateCategories;
+	SGameProjectDialog::GetAllTemplateCategories(AllTemplateCategories);
+
 	for (const TSharedPtr<FTemplateCategory>& Category : AllTemplateCategories)
 	{
 		if (Category->IsMajor)
@@ -153,32 +161,8 @@ TSharedRef<SWidget> SGameProjectDialog::CreateLandingPage()
 		}
 	}
 
-	FText Description;
-
-	if (DialogMode == EMode::New)
-	{
-		Description = LOCTEXT("LandingPageDesc_New", "Choose a template <RichTextBlock.BoldHighlight>category</> for your new project.");
-	}
-	else if (DialogMode == EMode::Open)
-	{
-		Description = LOCTEXT("LandingPageDesc_Open", "Choose a recent <RichTextBlock.BoldHighlight>project</> to load, or press <RichTextBlock.BoldHighlight>More</> to see all projects on your computer.");
-	}
-	else if (DialogMode == EMode::Both)
-	{
-		Description = LOCTEXT("LandingPageDesc_Open", "Choose a recent <RichTextBlock.BoldHighlight>project</> to load, press <RichTextBlock.BoldHighlight>More</> to see all projects on your computer, or choose a template <RichTextBlock.BoldHighlight>category</> to create a new project.");
-	}
-
 	TSharedRef<SWidget> Widget = SNew(SVerticalBox)
 	+ SVerticalBox::Slot()
-	.AutoHeight()
-	.Padding(0, 9, 0, 8)
-	[
-		SNew(SRichTextBlock)
-		.Text(Description)
-		.AutoWrapText(true)
-		.DecoratorStyleSet(&FEditorStyle::Get())
-	]
-	+ SVerticalBox::Slot()
 	.FillHeight(1)
 	[
 		SNew(SBorder)
diff --git a/Engine/Source/Editor/GameProjectGeneration/Private/SGameProjectDialog.h b/Engine/Source/Editor/GameProjectGeneration/Private/SGameProjectDialog.h
--- a/Engine/Source/Editor/GameProjectGeneration/Private/SGameProjectDialog.h
+++ b/Engine/Source/Editor/GameProjectGeneration/Private/SGameProjectDialog.h
@@ -43,6 +43,8 @@ public:
 	/** Constructs this widget with InArgs */
 	void Construct( const FArguments& InArgs, EMode Mode );
 
+	static void GetAllTemplateCategories(TArray<TSharedPtr<FTemplateCategory>>& OutCategories);
+
 private:
 	EMode DialogMode;
 
diff --git a/Engine/Source/Editor/GameProjectGeneration/Private/SNewProjectWizard.cpp b/Engine/Source/Editor/GameProjectGeneration/Private/SNewProjectWizard.cpp
--- a/Engine/Source/Editor/GameProjectGeneration/Private/SNewProjectWizard.cpp
+++ b/Engine/Source/Editor/GameProjectGeneration/Private/SNewProjectWizard.cpp
@@ -1,50 +1,50 @@
 // Copyright 1998-2019 Epic Games, Inc. All Rights Reserved.
 
-
 #include "SNewProjectWizard.h"
 #include "Brushes/SlateDynamicImageBrush.h"
-#include "HAL/PlatformFilemanager.h"
-#include "Misc/MessageDialog.h"
+#include "DesktopPlatformModule.h"
+#include "Dialogs/SOutputLogDialog.h"
+#include "Editor.h"
+#include "EditorStyleSet.h"
+#include "Framework/Application/SlateApplication.h"
+#include "GameProjectGenerationLog.h"
+#include "GameProjectGenerationModule.h"
+#include "GameProjectUtils.h"
 #include "HAL/FileManager.h"
+#include "HAL/PlatformFilemanager.h"
+#include "HardwareTargetingModule.h"
+#include "IDocumentation.h"
+#include "Interfaces/IPluginManager.h"
+#include "Interfaces/IProjectManager.h"
+#include "Internationalization/BreakIterator.h"
+#include "Layout/WidgetPath.h"
 #include "Misc/App.h"
 #include "Misc/CommandLine.h"
+#include "Misc/MessageDialog.h"
 #include "Misc/Parse.h"
-#include "Widgets/SOverlay.h"
-#include "Layout/WidgetPath.h"
+#include "ProjectDescriptor.h"
+#include "SGameProjectDialog.h"
+#include "SGetSuggestedIDEWidget.h"
+#include "Settings/EditorSettings.h"
 #include "SlateOptMacros.h"
-#include "Framework/Application/SlateApplication.h"
+#include "SourceCodeNavigation.h"
+#include "TemplateCategory.h"
+#include "TemplateItem.h"
+#include "TemplateProjectDefs.h"
 #include "Textures/SlateIcon.h"
-#include "Widgets/Layout/SSeparator.h"
 #include "Widgets/Images/SImage.h"
-#include "Widgets/Text/SRichTextBlock.h"
-#include "Widgets/Input/SEditableTextBox.h"
 #include "Widgets/Input/SButton.h"
+#include "Widgets/Input/SCheckBox.h"
+#include "Widgets/Input/SEditableTextBox.h"
 #include "Widgets/Layout/SGridPanel.h"
 #include "Widgets/Layout/SScrollBorder.h"
 #include "Widgets/Layout/SScrollBox.h"
+#include "Widgets/Layout/SSeparator.h"
 #include "Widgets/Layout/SWrapBox.h"
-#include "Widgets/Input/SCheckBox.h"
-#include "EditorStyleSet.h"
-#include "Editor.h"
-#include "Interfaces/IPluginManager.h"
-#include "Interfaces/IProjectManager.h"
-#include "ProjectDescriptor.h"
-#include "GameProjectGenerationLog.h"
-#include "GameProjectGenerationModule.h"
-#include "HardwareTargetingModule.h"
-#include "TemplateProjectDefs.h"
-#include "GameProjectUtils.h"
-#include "SGetSuggestedIDEWidget.h"
-#include "DesktopPlatformModule.h"
-#include "SourceCodeNavigation.h"
-#include "TemplateCategory.h"
+#include "Widgets/SOverlay.h"
 #include "Widgets/SToolTip.h"
+#include "Widgets/Text/SRichTextBlock.h"
 #include "Widgets/Workflow/SWizard.h"
-#include "IDocumentation.h"
-#include "Internationalization/BreakIterator.h"
-#include "Dialogs/SOutputLogDialog.h"
-#include "TemplateItem.h"
-#include "Settings/EditorSettings.h"
 
 #define LOCTEXT_NAMESPACE "NewProjectWizard"
 
@@ -339,16 +339,5 @@
 		SNew(SVerticalBox)
 
 		+ SVerticalBox::Slot()
-		.Padding(0, 8, 0, 8)
-		.AutoHeight()
 		[
-			SNew(SRichTextBlock)
-			.Text(LOCTEXT("ProjectTemplateDescription", "Choose a <RichTextBlock.BoldHighlight>template</> to use as a starting point for your new project.  Any of these features can be added later by clicking <RichTextBlock.BoldHighlight>Add Feature or Content Pack</> in <RichTextBlock.BoldHighlight>Content Browser</>."))
-			.AutoWrapText(true)
-			.DecoratorStyleSet(&FEditorStyle::Get())
-			.ToolTip(IDocumentation::Get()->CreateToolTip(LOCTEXT("TemplateChoiceTooltip", "A template consists of a little bit of player control logic (either as a Blueprint or in C++), input bindings, and appropriate prototyping assets."), nullptr, TEXT("Shared/Editor/NewProjectWizard"), TEXT("TemplateChoice")))
-		]
-
-		+ SVerticalBox::Slot()
-		[
 			SNew(SVerticalBox)
@@ -935,7 +924,7 @@
 		BlankTemplate->bIsBlankTemplate = true;
 
 		TArray<TSharedPtr<FTemplateCategory>> AllTemplateCategories;
-		FGameProjectGenerationModule::Get().GetAllTemplateCategories(AllTemplateCategories);
+		SGameProjectDialog::GetAllTemplateCategories(AllTemplateCategories);
 
 		for (const TSharedPtr<FTemplateCategory>& Category : AllTemplateCategories)
 		{
@@ -1314,12 +1303,6 @@
 	[
 		SNew(SVerticalBox)
 		+ SVerticalBox::Slot()
-		.AutoHeight()
-		.Padding(0, 8, 0, 8)
-		[
-			MakeProjectSettingsDescriptionBox()
-		]
-		+ SVerticalBox::Slot()
 		.Padding(0)
 		.VAlign(VAlign_Fill)
 		[
@@ -1462,17 +1445,6 @@
 	bShouldGenerateCode = Index == 1;
 }
 
-TSharedRef<SRichTextBlock> SNewProjectWizard::MakeProjectSettingsDescriptionBox() const
-{
-	TSharedRef<SRichTextBlock> Widget = SNew(SRichTextBlock)
-		.Text(LOCTEXT("ProjectSettingsDescription", "Choose some <RichTextBlock.BoldHighlight>settings</> for your project."))
-		.AutoWrapText(true)
-		.DecoratorStyleSet(&FEditorStyle::Get())
-		.ToolTip(IDocumentation::Get()->CreateToolTip(LOCTEXT("HardwareTargetTooltip", "These settings will choose good defaults for a number of other settings in the project such as post-processing flags and touch input emulation using the mouse."), nullptr, TEXT("Shared/Editor/NewProjectWizard"), TEXT("TargetHardware")));
-
-	return Widget;
-}
-
 static void AddToProjectSettingsGrid(TSharedRef<SGridPanel> Grid, const TSharedRef<SWidget>& Enum, const TSharedRef<SWidget>& Description, FIntPoint& Slot)
 {
 	Grid->AddSlot(Slot.X, Slot.Y)
diff --git a/Engine/Source/Editor/GameProjectGeneration/Private/SNewProjectWizard.h b/Engine/Source/Editor/GameProjectGeneration/Private/SNewProjectWizard.h
--- a/Engine/Source/Editor/GameProjectGeneration/Private/SNewProjectWizard.h
+++ b/Engine/Source/Editor/GameProjectGeneration/Private/SNewProjectWizard.h
@@ -203,7 +203,6 @@ private:
 	/** Handler for when the selection changes in the template list */
 	void HandleTemplateListViewSelectionChanged( TSharedPtr<FTemplateItem> TemplateItem, ESelectInfo::Type SelectInfo );
 
-	TSharedRef<SRichTextBlock> MakeProjectSettingsDescriptionBox() const;
 	TSharedRef<SWidget> MakeProjectSettingsOptionsBox();
 
 private:
