Removed unused variables from GameProjectUtils::AddSharedContentToProject().

diff --git a/Engine/Source/Editor/GameProjectGeneration/Private/GameProjectUtils.cpp b/Engine/Source/Editor/GameProjectGeneration/Private/GameProjectUtils.cpp
--- a/Engine/Source/Editor/GameProjectGeneration/Private/GameProjectUtils.cpp
+++ b/Engine/Source/Editor/GameProjectGeneration/Private/GameProjectUtils.cpp
@@ -4144,9 +4144,6 @@ bool GameProjectUtils::AddSharedContentToProject(const FProjectInformation &InPr
 	const FString SrcFolder = FPaths::GetPath(InProjectInfo.TemplateFile);
 	const FString DestFolder = FPaths::GetPath(InProjectInfo.ProjectFilename);
 
-	const FString ProjectConfigPath = DestFolder / TEXT("Config");
-	const FString IniFilename = ProjectConfigPath / TEXT("DefaultGame.ini");
-
 	// Now any packs specified in the template def.
 	UTemplateProjectDefs* TemplateDefs = LoadTemplateDefs(SrcFolder);
 	if (TemplateDefs != NULL)
@@ -4158,7 +4155,7 @@ bool GameProjectUtils::AddSharedContentToProject(const FProjectInformation &InPr
 		}
 
 		TUniquePtr<FFeaturePackContentSource> TempFeaturePack = MakeUnique<FFeaturePackContentSource>();
-		bool bCopied = TempFeaturePack->InsertAdditionalResources(TemplateDefs->SharedContentPacks,RequiredDetail, DestFolder,CreatedFiles);
+		bool bCopied = TempFeaturePack->InsertAdditionalResources(TemplateDefs->SharedContentPacks, RequiredDetail, DestFolder, CreatedFiles);
 		if( bCopied == false )
 		{
 			FFormatNamedArguments Args;
