Moved missing template category warning from template loading to New Project dialog creation.

diff --git a/Engine/Source/Editor/GameProjectGeneration/Private/GameProjectUtils.cpp b/Engine/Source/Editor/GameProjectGeneration/Private/GameProjectUtils.cpp
--- a/Engine/Source/Editor/GameProjectGeneration/Private/GameProjectUtils.cpp
+++ b/Engine/Source/Editor/GameProjectGeneration/Private/GameProjectUtils.cpp
@@ -1338,22 +1338,6 @@ UTemplateProjectDefs* GameProjectUtils::LoadTemplateDefs(const FString& ProjectD
 		{
 			UE_LOG(LogGameProjectGeneration, Warning, TEXT("Template '%s' contains 'All' in HiddenSettings in addition to other entries. This is a mistake, and means that all settings will be hidden."), *ProjectDirectory);
 		}
-
-		TArray<TSharedPtr<FTemplateCategory>> AllTemplateCategories;
-		FGameProjectGenerationModule::Get().GetAllTemplateCategories(AllTemplateCategories);
-
-		for (const FName& CategoryKey : TemplateDefs->Categories)
-		{
-			bool bCategoryExists = AllTemplateCategories.ContainsByPredicate([&CategoryKey](const TSharedPtr<FTemplateCategory>& Category)
-				{
-					return Category->Key == CategoryKey;
-				});
-
-			if (!bCategoryExists)
-			{
-				UE_LOG(LogGameProjectGeneration, Warning, TEXT("Failed to find category definition named '%s' while loading template '%s', it is not defined in any TemplateCategories.ini."), *CategoryKey.ToString(), *ProjectDirectory);
-			}
-		}
 	}
 
 	return TemplateDefs;
diff --git a/Engine/Source/Editor/GameProjectGeneration/Private/SNewProjectWizard.cpp b/Engine/Source/Editor/GameProjectGeneration/Private/SNewProjectWizard.cpp
--- a/Engine/Source/Editor/GameProjectGeneration/Private/SNewProjectWizard.cpp
+++ b/Engine/Source/Editor/GameProjectGeneration/Private/SNewProjectWizard.cpp
@@ -903,6 +903,26 @@ TMap<FName, TArray<TSharedPtr<FTemplateItem>> >& SNewProjectWizard::FindTemplate
 		}
 	}
 
+	TArray<TSharedPtr<FTemplateCategory>> AllTemplateCategories;
+	SGameProjectDialog::GetAllTemplateCategories(AllTemplateCategories);
+
+	// Validate that all our templates have a category defined
+	TArray<FName> CategoryKeys;
+	Templates.GetKeys(CategoryKeys);
+	for (const FName& CategoryKey : CategoryKeys)
+	{
+		bool bCategoryExists = AllTemplateCategories.ContainsByPredicate([&CategoryKey](const TSharedPtr<FTemplateCategory>& Category)
+			{
+				return Category->Key == CategoryKey;
+			});
+
+		if (!bCategoryExists)
+		{
+			UE_LOG(LogGameProjectGeneration, Warning, TEXT("Failed to find category definition named '%s', it is not defined in any TemplateCategories.ini."), *CategoryKey.ToString());
+		}
+	}
+
+
 	if (Templates.Num() == 0)
 	{
 		static const FName DefaultCategory("Default");
@@ -923,9 +943,6 @@ TMap<FName, TArray<TSharedPtr<FTemplateItem>> >& SNewProjectWizard::FindTemplate
 		BlankTemplate->bIsEnterprise = false;
 		BlankTemplate->bIsBlankTemplate = true;
 
-		TArray<TSharedPtr<FTemplateCategory>> AllTemplateCategories;
-		SGameProjectDialog::GetAllTemplateCategories(AllTemplateCategories);
-
 		for (const TSharedPtr<FTemplateCategory>& Category : AllTemplateCategories)
 		{
 			const TArray<TSharedPtr<FTemplateItem>>* CategoryEntry = Templates.Find(Category->Key);
