Changed project creation wizard titles to be clearer.

diff --git a/Engine/Source/Editor/GameProjectGeneration/Private/SGameProjectDialog.cpp b/Engine/Source/Editor/GameProjectGeneration/Private/SGameProjectDialog.cpp
--- a/Engine/Source/Editor/GameProjectGeneration/Private/SGameProjectDialog.cpp
+++ b/Engine/Source/Editor/GameProjectGeneration/Private/SGameProjectDialog.cpp
@@ -51,9 +51,9 @@ namespace GameProjectDialogDefs
 	constexpr int32 TemplateListPageIndex = 2;
 	constexpr int32 ProjectSettingsPageIndex = 3;
 	
-	static const FText LandingPageTitle = LOCTEXT("ProjectDialog_SelectProject", "Select Project");
+	static const FText LandingPageTitle = LOCTEXT("ProjectDialog_SelectOrCreateProject", "Select or Create New Project");
 	static const FText SelectCategoryTitle = LOCTEXT("ProjectDialog_SelectTemplateCategory", "Select Template Category");
-	static const FText ProjectBrowserTitle = LOCTEXT("ProjectDialog_ProjectBrowser", "Project Browser");
+	static const FText ProjectBrowserTitle = LOCTEXT("ProjectDialog_ProjectBrowserTitle", "Open Existing Project");
 	static const FText TemplateListTitle = LOCTEXT("ProjectDialog_SelectTemplate", "Select Template");
 	static const FText ProjectSettingsTitle = LOCTEXT("ProjectDialog_ProjectSettings", "Project Settings");
 }
@@ -243,7 +243,7 @@ TSharedRef<SWidget> SGameProjectDialog::CreateLandingPage()
 				[
 					SNew(STextBlock)
 					.TextStyle(FEditorStyle::Get(), "GameProjectDialog.ProjectNamePathLabels")
-					.Text(LOCTEXT("ProjectDialog_Categories", "Categories"))
+					.Text(LOCTEXT("ProjectDialog_NewProjectCategories", "New Project Categories"))
 					.Visibility(DialogMode == EMode::Both ? EVisibility::Visible : EVisibility::Collapsed)
 				]
 				+ SVerticalBox::Slot()
