Added option to set a custom StarterContent feature pack in a template, still uses the same default starter content packs if the template does not define one. This is to allow different sets of starter content to be specified for different problem domains such as archviz or virtual production.

Also create an empty Content folder always on project creation - this is to support template-based blank projects, which otherwise would not populate the Content folder at all, causing the Content folder directory watcher to not initialize on startup.

diff --git a/Engine/Source/Editor/GameProjectGeneration/Classes/TemplateProjectDefs.h b/Engine/Source/Editor/GameProjectGeneration/Classes/TemplateProjectDefs.h
--- a/Engine/Source/Editor/GameProjectGeneration/Classes/TemplateProjectDefs.h
+++ b/Engine/Source/Editor/GameProjectGeneration/Classes/TemplateProjectDefs.h
@@ -145,6 +145,9 @@ public:
 	UPROPERTY(config)
 	TArray<FFeaturePackLevelSet> SharedContentPacks;
 
+	UPROPERTY(config)
+	FString StarterContent;
+
 	/** Fixes up all strings in this definitions object to replace \%TEMPLATENAME\% with the supplied template name and \%PROJECTNAME\% with the supplied project name */
 	void FixupStrings(const FString& TemplateName, const FString& ProjectName);
 
diff --git a/Engine/Source/Editor/GameProjectGeneration/Private/GameProjectUtils.cpp b/Engine/Source/Editor/GameProjectGeneration/Private/GameProjectUtils.cpp
--- a/Engine/Source/Editor/GameProjectGeneration/Private/GameProjectUtils.cpp
+++ b/Engine/Source/Editor/GameProjectGeneration/Private/GameProjectUtils.cpp
@@ -209,7 +209,7 @@ namespace
 
 		if (InProjectInfo.bIsBlankTemplate &&
 			InProjectInfo.bCopyStarterContent &&
-			GameProjectUtils::IsStarterContentAvailableForNewProjects())
+			GameProjectUtils::IsStarterContentAvailableForProject(InProjectInfo))
 		{
 			DefaultMapConfig += LINE_TERMINATOR;
 			DefaultMapConfig += TEXT("[/Script/EngineSettings.GameMapsSettings]") LINE_TERMINATOR;
@@ -821,11 +821,27 @@ bool GameProjectUtils::OpenCodeIDE(const FString& ProjectFile, FText& OutFailRea
 	return true;
 }
 
-bool GameProjectUtils::IsStarterContentAvailableForNewProjects()
+FString GameProjectUtils::GetStarterContentName(const FProjectInformation& InProjectInfo)
 {
-	TArray<FString> OutFilenames;
-	IFileManager::Get().FindFilesRecursive(OutFilenames, *FPaths::FeaturePackDir(), TEXT("*StarterContent.upack"), /*Files=*/true, /*Directories=*/false);
-	return OutFilenames.Num() > 0;
+	if (!InProjectInfo.StarterContent.IsEmpty())
+	{
+		return InProjectInfo.StarterContent;
+	}
+
+	if (InProjectInfo.TargetedHardware == EHardwareClass::Mobile)
+	{
+		return TEXT("MobileStarterContent");
+	}
+
+	return TEXT("StarterContent");
+}
+
+bool GameProjectUtils::IsStarterContentAvailableForProject(const FProjectInformation& InProjectInfo)
+{
+	const FString StarterContentName = GameProjectUtils::GetStarterContentName(InProjectInfo);
+	const FString StarterContentPackFilename = FPaths::FeaturePackDir() / FString::Printf(TEXT("%s.upack"), *StarterContentName);
+
+	return IFileManager::Get().FileExists(*StarterContentPackFilename);
 }
 
 bool GameProjectUtils::CreateProject(const FProjectInformation& InProjectInfo, FText& OutFailReason, FText& OutFailLog, TArray<FString>* OutCreatedFiles)
@@ -1383,16 +1399,6 @@ bool GameProjectUtils::GenerateProjectFromScratch(const FProjectInformation& InP
 		return false;
 	}
 
-	// Make the Content folder
-	const FString ContentFolder = NewProjectFolder / TEXT("Content");
-	if ( !IFileManager::Get().MakeDirectory(*ContentFolder) )
-	{
-		FFormatNamedArguments Args;
-		Args.Add( TEXT("ContentFolder"), FText::FromString( ContentFolder ) );
-		OutFailReason = FText::Format( LOCTEXT("FailedToCreateContentFolder", "Failed to create the content folder {ContentFolder}"), Args );
-		return false;
-	}
-
 	SlowTask.EnterProgressFrame();
 
 	TArray<FString> StartupModuleNames;
@@ -1494,6 +1500,16 @@ bool GameProjectUtils::CreateProjectFromTemplate(const FProjectInformation& InPr
 
 	SlowTask.EnterProgressFrame();
 
+	// create a Content folder in the destination directory
+	const FString DestContentFolder = DestFolder / TEXT("Content");
+	if (!IFileManager::Get().MakeDirectory(*DestContentFolder, true))
+	{
+		FFormatNamedArguments FailArgs;
+		FailArgs.Add(TEXT("DestContentFolder"), FText::FromString(DestContentFolder));
+		OutFailReason = FText::Format(LOCTEXT("FailedToCreateDirectory", "Failed to create directory \"{DestContentFolder}\"."), FailArgs);
+		return false;
+	}
+
 	// Add a rule to replace the build settings version with the appropriate number
 	FTemplateReplacement& DefaultBuildSettingsRepl = TemplateDefs->ReplacementsInFiles.AddZeroed_GetRef();
 	DefaultBuildSettingsRepl.Extensions.Add("cs");
@@ -2063,18 +2079,12 @@ bool GameProjectUtils::GenerateConfigFiles(const FProjectInformation& InProjectI
 			FileContents += LINE_TERMINATOR;
 			FileContents += TEXT("[/Script/EngineSettings.GameMapsSettings]") LINE_TERMINATOR;
 
-			if (GameProjectUtils::IsStarterContentAvailableForNewProjects())
+			if (GameProjectUtils::IsStarterContentAvailableForProject(InProjectInfo))
 			{
-				if (InProjectInfo.TargetedHardware == EHardwareClass::Mobile)
-				{
-					FileContents += TEXT("EditorStartupMap=/Game/MobileStarterContent/Maps/Minimal_Default") LINE_TERMINATOR;
-					FileContents += TEXT("GameDefaultMap=/Game/MobileStarterContent/Maps/Minimal_Default") LINE_TERMINATOR;
-				}
-				else
-				{
-					FileContents += TEXT("EditorStartupMap=/Game/StarterContent/Maps/Minimal_Default") LINE_TERMINATOR;
-					FileContents += TEXT("GameDefaultMap=/Game/StarterContent/Maps/Minimal_Default") LINE_TERMINATOR;
-				}
+				FString StarterContentName = GameProjectUtils::GetStarterContentName(InProjectInfo);
+
+				FileContents += FString::Printf(TEXT("EditorStartupMap=/Game/%s/Maps/Minimal_Default") LINE_TERMINATOR, *StarterContentName);
+				FileContents += FString::Printf(TEXT("GameDefaultMap=/Game/%s/Maps/Minimal_Default") LINE_TERMINATOR, *StarterContentName);
 			}
 
 			if (InProjectInfo.bShouldGenerateCode)
@@ -2729,7 +2739,7 @@ bool GameProjectUtils::GetClassLocation(const FString& InPath, const FModuleCont
 
 GameProjectUtils::EProjectDuplicateResult GameProjectUtils::DuplicateProjectForUpgrade( const FString& InProjectFile, FString& OutNewProjectFile )
 {
-	IPlatformFile &PlatformFile = FPlatformFileManager::Get().GetPlatformFile();
+	IPlatformFile& PlatformFile = FPlatformFileManager::Get().GetPlatformFile();
 
 	// Get the directory part of the project name
 	FString OldDirectoryName = FPaths::GetPath(InProjectFile);
@@ -2881,17 +2891,17 @@ GameProjectUtils::EProjectDuplicateResult GameProjectUtils::DuplicateProjectForU
 void GameProjectUtils::UpdateSupportedTargetPlatforms(const FName& InPlatformName, const bool bIsSupported)
 {
 	const FString& ProjectFilename = FPaths::GetProjectFilePath();
-	if(!ProjectFilename.IsEmpty())
+	if (!ProjectFilename.IsEmpty())
 	{
 		// First attempt to check out the file if SCC is enabled
-		if(ISourceControlModule::Get().IsEnabled())
+		if (ISourceControlModule::Get().IsEnabled())
 		{
 			FText UnusedFailReason;
 			CheckoutGameProjectFile(ProjectFilename, UnusedFailReason);
 		}
 
 		// Second make sure the file is writable
-		if(FPlatformFileManager::Get().GetPlatformFile().IsReadOnly(*ProjectFilename))
+		if (FPlatformFileManager::Get().GetPlatformFile().IsReadOnly(*ProjectFilename))
 		{
 			FPlatformFileManager::Get().GetPlatformFile().SetReadOnly(*ProjectFilename, false);
 		}
@@ -2903,17 +2913,17 @@ void GameProjectUtils::UpdateSupportedTargetPlatforms(const FName& InPlatformNam
 void GameProjectUtils::ClearSupportedTargetPlatforms()
 {
 	const FString& ProjectFilename = FPaths::GetProjectFilePath();
-	if(!ProjectFilename.IsEmpty())
+	if (!ProjectFilename.IsEmpty())
 	{
 		// First attempt to check out the file if SCC is enabled
-		if(ISourceControlModule::Get().IsEnabled())
+		if (ISourceControlModule::Get().IsEnabled())
 		{
 			FText UnusedFailReason;
 			CheckoutGameProjectFile(ProjectFilename, UnusedFailReason);
 		}
 
 		// Second make sure the file is writable
-		if(FPlatformFileManager::Get().GetPlatformFile().IsReadOnly(*ProjectFilename))
+		if (FPlatformFileManager::Get().GetPlatformFile().IsReadOnly(*ProjectFilename))
 		{
 			FPlatformFileManager::Get().GetPlatformFile().SetReadOnly(*ProjectFilename, false);
 		}
@@ -3716,11 +3726,6 @@ bool GameProjectUtils::CheckoutGameProjectFile(const FString& ProjectFilename, F
 	return bSuccessfullyCheckedOut;
 }
 
-FString GameProjectUtils::GetDefaultProjectTemplateFilename()
-{
-	return TEXT("");
-}
-
 void FindCodeFiles(const TCHAR* BaseDirectory, TArray<FString>& FileNames, int32 MaxNumFileNames)
 {
 	struct FDirectoryVisitor : public IPlatformFile::FDirectoryVisitor
@@ -4115,15 +4120,8 @@ bool GameProjectUtils::InsertFeaturePacksIntoINIFile(const FProjectInformation&
 	// First the starter content
 	if (InProjectInfo.bCopyStarterContent)
 	{
-		FString StarterPack;
-		if (InProjectInfo.TargetedHardware == EHardwareClass::Mobile)
-		{
-			StarterPack = TEXT("InsertPack=(PackSource=\"MobileStarterContent.upack\",PackName=\"StarterContent\")");
-		}
-		else
-		{
-			StarterPack = TEXT("InsertPack=(PackSource=\"StarterContent.upack\",PackName=\"StarterContent\")");
-		}
+		FString StarterContentName = GameProjectUtils::GetStarterContentName(InProjectInfo);
+		FString StarterPack = FString::Printf(TEXT("InsertPack=(PackSource=\"%s.upack\",PackName=\"StarterContent\")"), *StarterContentName);
 		PackList.Add(StarterPack);
 	}
 
diff --git a/Engine/Source/Editor/GameProjectGeneration/Private/SNewProjectWizard.cpp b/Engine/Source/Editor/GameProjectGeneration/Private/SNewProjectWizard.cpp
--- a/Engine/Source/Editor/GameProjectGeneration/Private/SNewProjectWizard.cpp
+++ b/Engine/Source/Editor/GameProjectGeneration/Private/SNewProjectWizard.cpp
@@ -290,7 +290,7 @@ private:
 	/** Get this item's thumbnail or return the default */
 	const FSlateBrush* GetThumbnail() const
 	{
-		auto ItemPtr = Item.Pin();
+		TSharedPtr<FTemplateItem> ItemPtr = Item.Pin();
 		if (ItemPtr.IsValid() && ItemPtr->Thumbnail.IsValid())
 		{
 			return ItemPtr->Thumbnail.Get();
@@ -326,8 +326,6 @@ void SNewProjectWizard::Construct( const FArguments& InArgs )
 	.OnMouseButtonDoubleClick(this, &SNewProjectWizard::HandleTemplateListViewDoubleClick)
 	.OnSelectionChanged(this, &SNewProjectWizard::HandleTemplateListViewSelectionChanged);
 
-	const EVisibility StarterContentVisiblity = GameProjectUtils::IsStarterContentAvailableForNewProjects() ? EVisibility::Visible : EVisibility::Collapsed;
-
 	TSharedRef<SSeparator> Separator = SNew(SSeparator).Orientation(EOrientation::Orient_Vertical);
 	Separator->SetBorderBackgroundColor(FLinearColor::White.CopyWithNewOpacity(0.25f));
 
@@ -565,13 +563,13 @@ EVisibility SNewProjectWizard::GetSelectedTemplateAssetVisibility() const
 
 const FSlateBrush* SNewProjectWizard::GetSelectedTemplatePreviewImage() const
 {
-	auto PreviewImage = GetSelectedTemplateProperty(&FTemplateItem::PreviewImage);
+	TSharedPtr<FSlateBrush> PreviewImage = GetSelectedTemplateProperty(&FTemplateItem::PreviewImage);
 	return PreviewImage.IsValid() ? PreviewImage.Get() : nullptr;
 }
 
 EVisibility SNewProjectWizard::GetSelectedTemplatePreviewVisibility() const
 {
-	auto PreviewImage = GetSelectedTemplateProperty(&FTemplateItem::PreviewImage);
+	TSharedPtr<FSlateBrush> PreviewImage = GetSelectedTemplateProperty(&FTemplateItem::PreviewImage);
 	return PreviewImage.IsValid() ? EVisibility::Visible : EVisibility::Collapsed;
 }
 
@@ -743,7 +741,7 @@ TMap<FName, TArray<TSharedPtr<FTemplateItem>> >& SNewProjectWizard::FindTemplate
 	// Add the Enterprise templates
 	TemplateRootFolders.Add(FPaths::EnterpriseDir() + TEXT("Templates"));
 
-	// allow plugins to define templates
+	// Allow plugins to define templates
 	TArray<TSharedRef<IPlugin>> Plugins = IPluginManager::Get().GetEnabledPlugins();
 	for (const TSharedRef<IPlugin>& Plugin : Plugins)
 	{
@@ -761,136 +759,132 @@ TMap<FName, TArray<TSharedPtr<FTemplateItem>> >& SNewProjectWizard::FindTemplate
 
 	// Form a list of all folders that could contain template projects
 	TArray<FString> AllTemplateFolders;
-	for ( auto TemplateRootFolderIt = TemplateRootFolders.CreateConstIterator(); TemplateRootFolderIt; ++TemplateRootFolderIt )
+	for (const FString& Root : TemplateRootFolders)
 	{
-		const FString Root = *TemplateRootFolderIt;
-		
 		const FString SearchString = Root / TEXT("*");
 		TArray<FString> TemplateFolders;
 		IFileManager::Get().FindFiles(TemplateFolders, *SearchString, /*Files=*/false, /*Directories=*/true);
-		for ( auto TemplateFolderIt = TemplateFolders.CreateConstIterator(); TemplateFolderIt; ++TemplateFolderIt )
+
+		for (const FString& Folder : TemplateFolders)
 		{
-			AllTemplateFolders.Add( Root / (*TemplateFolderIt) );
+			AllTemplateFolders.Add(Root / Folder);
 		}
 	}
 
 	TArray<TSharedPtr<FTemplateItem>> FoundTemplates;
 
 	// Add a template item for every discovered project
-	for ( auto TemplateFolderIt = AllTemplateFolders.CreateConstIterator(); TemplateFolderIt; ++TemplateFolderIt )
+	for (const FString& Root : AllTemplateFolders)
 	{
-		const FString SearchString = (*TemplateFolderIt) / TEXT("*.") + FProjectDescriptor::GetExtension();
+		const FString SearchString = Root / TEXT("*.") + FProjectDescriptor::GetExtension();
 		TArray<FString> FoundProjectFiles;
 		IFileManager::Get().FindFiles(FoundProjectFiles, *SearchString, /*Files=*/true, /*Directories=*/false);
-		if ( FoundProjectFiles.Num() > 0 )
+
+		if (FoundProjectFiles.Num() == 0 || !ensure(FoundProjectFiles.Num() == 1))
 		{
-			if ( ensure(FoundProjectFiles.Num() == 1) )
-			{
-				// Make sure a TemplateDefs ini file exists
-				const FString Root = *TemplateFolderIt;
-				UTemplateProjectDefs* TemplateDefs = GameProjectUtils::LoadTemplateDefs(Root);
-				if (TemplateDefs)
-				{
-					// Ignore any templates whose definition says we cannot use to create a project
-					if( TemplateDefs->bAllowProjectCreation == false )
-						continue;
-
-					FString ProjectFile = Root / FoundProjectFiles[0];
-
-					// If no template name was specified for the current culture, just use the project name
-					TArray<FName> TemplateCategories = TemplateDefs->Categories;
-					if (TemplateCategories.Num() == 0)
-					{
-						TemplateCategories.Add(NewProjectWizardDefs::DefaultCategoryName);
-					}
-
-					FString TemplateKey = Root;
-					TemplateKey.RemoveFromEnd("BP");
-
-					TSharedPtr<FTemplateItem>* ExistingTemplate = FoundTemplates.FindByPredicate([&TemplateKey](TSharedPtr<FTemplateItem> Item)
-					{
-						return Item->Key == TemplateKey;
-					});
-
-					if (ExistingTemplate != nullptr)
-					{
-						if (TemplateDefs->GeneratesCode(Root))
-						{
-							(*ExistingTemplate)->CodeProjectFile = ProjectFile;
-						}
-						else
-						{
-							(*ExistingTemplate)->BlueprintProjectFile = ProjectFile;
-						}
-
-						continue;
-					}
+			continue;
+		}
+
+		// Make sure a TemplateDefs.ini file exists
+		UTemplateProjectDefs* TemplateDefs = GameProjectUtils::LoadTemplateDefs(Root);
+		if (TemplateDefs == nullptr)
+		{
+			continue;
+		}
+
+		// Ignore any templates whose definition says we cannot use to create a project
+		if (TemplateDefs->bAllowProjectCreation == false)
+		{
+			continue;
+		}
+
+		const FString ProjectFile = Root / FoundProjectFiles[0];
+
+		// If no template category was specified, use the default category
+		TArray<FName> TemplateCategories = TemplateDefs->Categories;
+		if (TemplateCategories.Num() == 0)
+		{
+			TemplateCategories.Add(NewProjectWizardDefs::DefaultCategoryName);
+		}
+
+		// Find a duplicate project, eg. "Foo" and "FooBP"
+		FString TemplateKey = Root;
+		TemplateKey.RemoveFromEnd("BP");
+
+		TSharedPtr<FTemplateItem>* ExistingTemplate = FoundTemplates.FindByPredicate([&TemplateKey](TSharedPtr<FTemplateItem> Item)
+		{
+			return Item->Key == TemplateKey;
+		});
+
+		TSharedPtr<FTemplateItem> Template;
+
+		// Create a new template if none was found
+		if (ExistingTemplate != nullptr)
+		{
+			Template = *ExistingTemplate;
+		}
+		else
+		{
+			Template = MakeShareable(new FTemplateItem());
+		}
+
+		if (TemplateDefs->GeneratesCode(Root))
+		{
+			Template->CodeProjectFile = ProjectFile;
+			Template->CodeTemplateDefs = TemplateDefs;
+		}
+		else
+		{
+			Template->BlueprintProjectFile = ProjectFile;
+			Template->BlueprintTemplateDefs = TemplateDefs;
+		}
+
+		// The rest has already been set by the existing template, so skip it.
+		if (ExistingTemplate != nullptr)
+		{
+			continue;
+		}
 					
-					// Did not find an existing template. Create a new one to add to the template list.
-					TSharedPtr<FTemplateItem> Template = MakeShareable(new FTemplateItem());
-					Template->Key = TemplateKey;
-					Template->Categories = TemplateCategories;
-					Template->Description = TemplateDefs->GetLocalizedDescription();
-					Template->ClassTypes = TemplateDefs->ClassTypes;
-					Template->AssetTypes = TemplateDefs->AssetTypes;
-					Template->HiddenSettings = TemplateDefs->HiddenSettings;
-					Template->bIsEnterprise = TemplateDefs->bIsEnterprise;
-					Template->bIsBlankTemplate = TemplateDefs->bIsBlank;
-
-					Template->Name = TemplateDefs->GetDisplayNameText();
-					if (Template->Name.IsEmpty())
-					{
-						Template->Name = FText::FromString(FPaths::GetBaseFilename(ProjectFile));
-					}
-
-					// Only generate code if the template has a source folder
-					if (TemplateDefs->GeneratesCode(Root))
-					{
-						Template->CodeProjectFile = ProjectFile;
-					}
-					else
-					{
-						Template->BlueprintProjectFile = ProjectFile;
-					}
-
-					const FString ThumbnailPNGFile = (Root + TEXT("/Media/") + FoundProjectFiles[0]).Replace(TEXT(".uproject"), TEXT(".png"));
-					if ( FPlatformFileManager::Get().GetPlatformFile().FileExists(*ThumbnailPNGFile) )
-					{
-						const FName BrushName = FName(*ThumbnailPNGFile);
-						Template->Thumbnail = MakeShareable( new FSlateDynamicImageBrush(BrushName, FVector2D(128,128) ) );
-					}
-
-					TSharedPtr<FSlateDynamicImageBrush> PreviewBrush;
-					const FString PreviewPNGFile = (Root + TEXT("/Media/") + FoundProjectFiles[0]).Replace(TEXT(".uproject"), TEXT("_Preview.png"));
-					if ( FPlatformFileManager::Get().GetPlatformFile().FileExists(*PreviewPNGFile) )
-					{
-						const FName BrushName = FName(*PreviewPNGFile);
-						Template->PreviewImage = MakeShareable( new FSlateDynamicImageBrush(BrushName, FVector2D(512,256) ) );
-					}
-
-					// Get the sort key
-					FString SortKey = TemplateDefs->SortKey;
-					const FString CleanFilename = FPaths::GetCleanFilename(ProjectFile);
-					if(SortKey.Len() == 0)
-					{
-						SortKey = CleanFilename;
-					}
-
-					if (CleanFilename == GameProjectUtils::GetDefaultProjectTemplateFilename())
-					{
-						SortKey = TEXT("_0");
-					}
-					Template->SortKey = SortKey;
-
-					FoundTemplates.Add(Template);
-				}
-			}
-			else
-			{
-				// More than one project file in this template? This is not legal, skip it.
-				continue;
-			}
+		// Did not find an existing template. Create a new one to add to the template list.
+		Template->Key = TemplateKey;
+					
+		// @todo: These are all basically just copies of what's in UTemplateProjectDefs, but ignore differences between code and BP 
+		Template->Categories = TemplateCategories;
+		Template->Description = TemplateDefs->GetLocalizedDescription();
+		Template->ClassTypes = TemplateDefs->ClassTypes;
+		Template->AssetTypes = TemplateDefs->AssetTypes;
+		Template->HiddenSettings = TemplateDefs->HiddenSettings;
+		Template->bIsEnterprise = TemplateDefs->bIsEnterprise;
+		Template->bIsBlankTemplate = TemplateDefs->bIsBlank;
+
+		Template->Name = TemplateDefs->GetDisplayNameText();
+		if (Template->Name.IsEmpty())
+		{
+			Template->Name = FText::FromString(TemplateKey);
+		}
+
+		const FString ThumbnailPNGFile = (Root + TEXT("/Media/") + FoundProjectFiles[0]).Replace(TEXT(".uproject"), TEXT(".png"));
+		if ( FPlatformFileManager::Get().GetPlatformFile().FileExists(*ThumbnailPNGFile) )
+		{
+			const FName BrushName = FName(*ThumbnailPNGFile);
+			Template->Thumbnail = MakeShareable( new FSlateDynamicImageBrush(BrushName, FVector2D(128,128) ) );
 		}
+
+		TSharedPtr<FSlateDynamicImageBrush> PreviewBrush;
+		const FString PreviewPNGFile = (Root + TEXT("/Media/") + FoundProjectFiles[0]).Replace(TEXT(".uproject"), TEXT("_Preview.png"));
+		if ( FPlatformFileManager::Get().GetPlatformFile().FileExists(*PreviewPNGFile) )
+		{
+			const FName BrushName = FName(*PreviewPNGFile);
+			Template->PreviewImage = MakeShareable( new FSlateDynamicImageBrush(BrushName, FVector2D(512,256) ) );
+		}
+
+		Template->SortKey = TemplateDefs->SortKey;
+		if(Template->SortKey.IsEmpty())
+		{
+			Template->SortKey = TemplateKey;
+		}
+
+		FoundTemplates.Add(Template);
 	}
 
 	for (const TSharedPtr<FTemplateItem>& Template : FoundTemplates)
@@ -927,7 +921,7 @@ TMap<FName, TArray<TSharedPtr<FTemplateItem>> >& SNewProjectWizard::FindTemplate
 		Templates.Add(DefaultCategory);
 	}
 
-	// add blank template to empty categories
+	// Add blank template to empty categories
 	{
 		TSharedPtr<FTemplateItem> BlankTemplate = MakeShareable(new FTemplateItem());
 		BlankTemplate->Name = LOCTEXT("BlankProjectName", "Blank");
@@ -1095,23 +1089,15 @@ bool SNewProjectWizard::IsCompilerRequired() const
 	return false;
 }
 
-bool SNewProjectWizard::CreateProject(const FString& ProjectFile)
+FProjectInformation SNewProjectWizard::CreateProjectInfo() const
 {
-	// Get the selected template
 	TSharedPtr<FTemplateItem> SelectedTemplate = GetSelectedTemplateItem();
-
-	if (!ensure(SelectedTemplate.IsValid()))
+	if (!SelectedTemplate.IsValid())
 	{
-		// A template must be selected.
-		return false;
+		return FProjectInformation();
 	}
 
-	FText FailReason, FailLog;
-
-	const TArray<ETemplateSetting>& HiddenSettings = GetSelectedTemplateProperty(&FTemplateItem::HiddenSettings);
-
 	FProjectInformation ProjectInfo;
-	ProjectInfo.ProjectFilename = ProjectFile;
 	ProjectInfo.bShouldGenerateCode = bShouldGenerateCode;
 	ProjectInfo.bCopyStarterContent = bCopyStarterContent;
 	ProjectInfo.TemplateFile = bShouldGenerateCode ? SelectedTemplate->CodeProjectFile : SelectedTemplate->BlueprintProjectFile;
@@ -1119,6 +1105,9 @@ bool SNewProjectWizard::CreateProject(const FString& ProjectFile)
 	ProjectInfo.bIsEnterpriseProject = SelectedTemplate->bIsEnterprise;
 	ProjectInfo.bIsBlankTemplate = SelectedTemplate->bIsBlankTemplate;
 	ProjectInfo.bForceExtendedLuminanceRange = SelectedTemplate->bIsBlankTemplate;
+	ProjectInfo.StarterContent = bShouldGenerateCode ? SelectedTemplate->CodeTemplateDefs->StarterContent : SelectedTemplate->BlueprintTemplateDefs->StarterContent;
+
+	const TArray<ETemplateSetting>& HiddenSettings = SelectedTemplate->HiddenSettings;
 
 	if (!HiddenSettings.Contains(ETemplateSetting::All))
 	{
@@ -1143,6 +1132,25 @@ bool SNewProjectWizard::CreateProject(const FString& ProjectFile)
 		}
 	}
 
+	return MoveTemp(ProjectInfo);
+}
+
+bool SNewProjectWizard::CreateProject(const FString& ProjectFile)
+{
+	// Get the selected template
+	TSharedPtr<FTemplateItem> SelectedTemplate = GetSelectedTemplateItem();
+
+	if (!ensure(SelectedTemplate.IsValid()))
+	{
+		// A template must be selected.
+		return false;
+	}
+
+	FText FailReason, FailLog;
+
+	FProjectInformation ProjectInfo = CreateProjectInfo();
+	ProjectInfo.ProjectFilename = ProjectFile;
+
 	if (!GameProjectUtils::CreateProject(ProjectInfo, FailReason, FailLog))
 	{
 		SOutputLogDialog::Open(LOCTEXT("CreateProject", "Create Project"), FailReason, FailLog, FText::GetEmpty());
@@ -1152,13 +1160,13 @@ bool SNewProjectWizard::CreateProject(const FString& ProjectFile)
 	// Successfully created the project. Update the last created location string.
 	FString CreatedProjectPath = FPaths::GetPath(FPaths::GetPath(ProjectFile)); 
 
-	// if the original path was the drives root (ie: C:/) the double path call strips the last /
+	// If the original path was the drives root (ie: C:/) the double path call strips the last /
 	if (CreatedProjectPath.EndsWith(":"))
 	{
 		CreatedProjectPath.AppendChar('/');
 	}
 
-	auto* Settings = GetMutableDefault<UEditorSettings>();
+	UEditorSettings* Settings = GetMutableDefault<UEditorSettings>();
 	Settings->CreatedProjectPaths.Remove(CreatedProjectPath);
 	Settings->CreatedProjectPaths.Insert(CreatedProjectPath, 0);
 	Settings->bCopyStarterContentPreference = bCopyStarterContent;
@@ -1597,8 +1605,10 @@ TSharedRef<SWidget> SNewProjectWizard::MakeProjectSettingsOptionsBox()
 			FSlateIcon(FEditorStyle::GetStyleSetName(), "GameProjectDialog.NoStarterContent"),
 			LOCTEXT("NoStarterContent", "No Starter Content")));
 
-		// Only add the option to add starter content if its there to add !
-		bool bIsStarterAvailable = GameProjectUtils::IsStarterContentAvailableForNewProjects();
+		// Only add the option to add starter content if its there to add!
+		FProjectInformation ProjectInfo = CreateProjectInfo();
+		bool bIsStarterAvailable = GameProjectUtils::IsStarterContentAvailableForProject(ProjectInfo);
+
 		StarterContentOptions.Add(SDecoratedEnumCombo<int32>::FComboOption(
 			1,
 			FSlateIcon(FEditorStyle::GetStyleSetName(), "GameProjectDialog.IncludeStarterContent"),
diff --git a/Engine/Source/Editor/GameProjectGeneration/Private/SNewProjectWizard.h b/Engine/Source/Editor/GameProjectGeneration/Private/SNewProjectWizard.h
--- a/Engine/Source/Editor/GameProjectGeneration/Private/SNewProjectWizard.h
+++ b/Engine/Source/Editor/GameProjectGeneration/Private/SNewProjectWizard.h
@@ -205,6 +205,9 @@ private:
 
 	TSharedRef<SWidget> MakeProjectSettingsOptionsBox();
 
+	/** Create a project information struct from the currently selected settings. */
+	struct FProjectInformation CreateProjectInfo() const;
+
 private:
 
 	FString LastBrowsePath;
diff --git a/Engine/Source/Editor/GameProjectGeneration/Private/TemplateItem.h b/Engine/Source/Editor/GameProjectGeneration/Private/TemplateItem.h
--- a/Engine/Source/Editor/GameProjectGeneration/Private/TemplateItem.h
+++ b/Engine/Source/Editor/GameProjectGeneration/Private/TemplateItem.h
@@ -23,10 +23,12 @@ struct FTemplateItem
 	FString		AssetTypes;
 
 	FString		CodeProjectFile;
+	UTemplateProjectDefs* CodeTemplateDefs = nullptr;
 	FString		BlueprintProjectFile;
+	UTemplateProjectDefs* BlueprintTemplateDefs = nullptr;
 
 	TArray<ETemplateSetting> HiddenSettings;
 
-	bool		bIsEnterprise;
-	bool		bIsBlankTemplate;
+	bool		bIsEnterprise = false;
+	bool		bIsBlankTemplate = false;
 };
diff --git a/Engine/Source/Editor/GameProjectGeneration/Public/GameProjectUtils.h b/Engine/Source/Editor/GameProjectGeneration/Public/GameProjectUtils.h
--- a/Engine/Source/Editor/GameProjectGeneration/Public/GameProjectUtils.h
+++ b/Engine/Source/Editor/GameProjectGeneration/Public/GameProjectUtils.h
@@ -35,6 +35,9 @@ struct FProjectInformation
 
 	TOptional<EHardwareClass::Type> TargetedHardware;
 	TOptional<EGraphicsPreset::Type> DefaultGraphicsPerformance;
+
+	/** The name of the feature pack to use as starter content. Must be located under FeaturePacks\. */
+	FString StarterContent;
 };
 
 DECLARE_DELEGATE_RetVal_OneParam(bool, FProjectDescriptorModifier, FProjectDescriptor&);
@@ -150,9 +153,6 @@ public:
 	*/
 	static void GetProjectSourceDirectoryInfo(int32& OutNumFiles, int64& OutDirectorySize);
 
-	/** Returns the uproject template filename for the default project template. */
-	static FString GetDefaultProjectTemplateFilename();
-
 	/** Compiles a project while showing a progress bar, and offers to open the IDE if it fails. */
 	static bool BuildCodeProject(const FString& ProjectFilename);
 
@@ -160,7 +160,7 @@ public:
 	static bool GenerateCodeProjectFiles(const FString& ProjectFilename, FText& OutFailReason, FText& OutFailLog);
 
 	/** Returns true if there are starter content files available for instancing into new projects. */
-	static bool IsStarterContentAvailableForNewProjects();
+	static bool IsStarterContentAvailableForProject(const FProjectInformation& ProjectInfo);
 
 	/**
 	 * Get the information about any modules referenced in the .uproject file of the currently loaded project
@@ -281,6 +281,10 @@ public:
 
 private:
 
+	/** Get the name of the starter content pack to use for the given project. */
+	static FString GetStarterContentName(const FProjectInformation& InProjectInfo);
+
+	/** Get the config string for the hardware target for the given project. */
 	static FString GetHardwareConfigString(const FProjectInformation& InProjectInfo);
 
 	/** Generates a new project without using a template project */
@@ -487,7 +491,6 @@ private:
 	 */
 	static bool UpdateRequiredAdditionalDependencies(FProjectDescriptor& Descriptor, TArray<FString>& RequiredDependencies, const FString& ModuleName);
 
-private:
 	/**
 	 * Updates the projects and modifies FProjectDescriptor accordingly to given modifier.
 	 *
@@ -507,6 +510,8 @@ private:
 	 */
 	static bool UpdateGameProjectFile_Impl(const FString& ProjectFilename, const FString& EngineIdentifier, const FProjectDescriptorModifier* Modifier, FText& OutFailReason);
 
+private:
+
 	static TWeakPtr<SNotificationItem> UpdateGameProjectNotification;
 	static TWeakPtr<SNotificationItem> WarningProjectNameNotification;
 
