The default map config string that gets added during project creation now respects the new StarterContent setting in FProjectInformation.

diff --git a/Engine/Source/Editor/GameProjectGeneration/Private/GameProjectUtils.cpp b/Engine/Source/Editor/GameProjectGeneration/Private/GameProjectUtils.cpp
--- a/Engine/Source/Editor/GameProjectGeneration/Private/GameProjectUtils.cpp
+++ b/Engine/Source/Editor/GameProjectGeneration/Private/GameProjectUtils.cpp
@@ -202,32 +202,6 @@ namespace
 
 		return MoveTemp(RaytracingConfig);
 	}
-
-	FString GetDefaultMapConfigString(const FProjectInformation& InProjectInfo)
-	{
-		FString DefaultMapConfig;
-
-		if (InProjectInfo.bIsBlankTemplate &&
-			InProjectInfo.bCopyStarterContent &&
-			GameProjectUtils::IsStarterContentAvailableForProject(InProjectInfo))
-		{
-			DefaultMapConfig += LINE_TERMINATOR;
-			DefaultMapConfig += TEXT("[/Script/EngineSettings.GameMapsSettings]") LINE_TERMINATOR;
-
-			if (InProjectInfo.TargetedHardware == EHardwareClass::Mobile)
-			{
-				DefaultMapConfig += TEXT("EditorStartupMap=/Game/MobileStarterContent/Maps/Minimal_Default") LINE_TERMINATOR;
-				DefaultMapConfig += TEXT("GameDefaultMap=/Game/MobileStarterContent/Maps/Minimal_Default") LINE_TERMINATOR;
-			}
-			else
-			{
-				DefaultMapConfig += TEXT("EditorStartupMap=/Game/StarterContent/Maps/Minimal_Default") LINE_TERMINATOR;
-				DefaultMapConfig += TEXT("GameDefaultMap=/Game/StarterContent/Maps/Minimal_Default") LINE_TERMINATOR;
-			}
-		}
-
-		return MoveTemp(DefaultMapConfig);
-	}
 } // namespace <>
 
 FText FNewClassInfo::GetClassName() const
@@ -844,6 +818,26 @@ bool GameProjectUtils::IsStarterContentAvailableForProject(const FProjectInforma
 	return IFileManager::Get().FileExists(*StarterContentPackFilename);
 }
 
+FString GameProjectUtils::GetDefaultMapConfigString(const FProjectInformation& InProjectInfo)
+{
+	FString DefaultMapConfig;
+
+	if (InProjectInfo.bIsBlankTemplate &&
+		InProjectInfo.bCopyStarterContent &&
+		GameProjectUtils::IsStarterContentAvailableForProject(InProjectInfo))
+	{
+		const FString StarterContentName = GameProjectUtils::GetStarterContentName(InProjectInfo);
+
+		DefaultMapConfig += LINE_TERMINATOR;
+		DefaultMapConfig += TEXT("[/Script/EngineSettings.GameMapsSettings]") LINE_TERMINATOR;
+
+		DefaultMapConfig += FString::Printf(TEXT("EditorStartupMap=/Game/%s/Maps/Minimal_Default") LINE_TERMINATOR, *StarterContentName);
+		DefaultMapConfig += FString::Printf(TEXT("GameDefaultMap=/Game/%s/Maps/Minimal_Default") LINE_TERMINATOR, *StarterContentName);
+	}
+
+	return MoveTemp(DefaultMapConfig);
+}
+
 bool GameProjectUtils::CreateProject(const FProjectInformation& InProjectInfo, FText& OutFailReason, FText& OutFailLog, TArray<FString>* OutCreatedFiles)
 {
 	if ( !IsValidProjectFileForCreation(InProjectInfo.ProjectFilename, OutFailReason) )
@@ -2074,16 +2068,7 @@ bool GameProjectUtils::GenerateConfigFiles(const FProjectInformation& InProjectI
 
 		if (InProjectInfo.bCopyStarterContent)
 		{
-			FileContents += LINE_TERMINATOR;
-			FileContents += TEXT("[/Script/EngineSettings.GameMapsSettings]") LINE_TERMINATOR;
-
-			if (GameProjectUtils::IsStarterContentAvailableForProject(InProjectInfo))
-			{
-				FString StarterContentName = GameProjectUtils::GetStarterContentName(InProjectInfo);
-
-				FileContents += FString::Printf(TEXT("EditorStartupMap=/Game/%s/Maps/Minimal_Default") LINE_TERMINATOR, *StarterContentName);
-				FileContents += FString::Printf(TEXT("GameDefaultMap=/Game/%s/Maps/Minimal_Default") LINE_TERMINATOR, *StarterContentName);
-			}
+			FileContents += GetDefaultMapConfigString(InProjectInfo);
 
 			if (InProjectInfo.bShouldGenerateCode)
 			{
diff --git a/Engine/Source/Editor/GameProjectGeneration/Public/GameProjectUtils.h b/Engine/Source/Editor/GameProjectGeneration/Public/GameProjectUtils.h
--- a/Engine/Source/Editor/GameProjectGeneration/Public/GameProjectUtils.h
+++ b/Engine/Source/Editor/GameProjectGeneration/Public/GameProjectUtils.h
@@ -287,6 +287,9 @@ private:
 	/** Get the config string for the hardware target for the given project. */
 	static FString GetHardwareConfigString(const FProjectInformation& InProjectInfo);
 
+	/** Get the config string to add for the default maps for the given project. */
+	static FString GetDefaultMapConfigString(const FProjectInformation& InProjectInfo);
+
 	/** Generates a new project without using a template project */
 	static bool GenerateProjectFromScratch(const FProjectInformation& InProjectInfo, FText& OutFailReason, FText& OutFailLog);
 
