Creating a new blank project without templates synced no longer crashes.

#rnx

diff --git a/Engine/Source/Editor/GameProjectGeneration/Private/SNewProjectWizard.cpp b/Engine/Source/Editor/GameProjectGeneration/Private/SNewProjectWizard.cpp
--- a/Engine/Source/Editor/GameProjectGeneration/Private/SNewProjectWizard.cpp
+++ b/Engine/Source/Editor/GameProjectGeneration/Private/SNewProjectWizard.cpp
@@ -1098,7 +1098,24 @@ FProjectInformation SNewProjectWizard::CreateProjectInfo() const
 	ProjectInfo.bIsEnterpriseProject = SelectedTemplate->bIsEnterprise;
 	ProjectInfo.bIsBlankTemplate = SelectedTemplate->bIsBlankTemplate;
 	ProjectInfo.bForceExtendedLuminanceRange = SelectedTemplate->bIsBlankTemplate;
-	ProjectInfo.StarterContent = bShouldGenerateCode ? SelectedTemplate->CodeTemplateDefs->StarterContent : SelectedTemplate->BlueprintTemplateDefs->StarterContent;
+
+	if (bCopyStarterContent)
+	{
+		if (bShouldGenerateCode)
+		{
+			if (SelectedTemplate->CodeTemplateDefs != nullptr)
+			{
+				ProjectInfo.StarterContent = SelectedTemplate->CodeTemplateDefs->StarterContent;
+			}
+		}
+		else
+		{
+			if (SelectedTemplate->BlueprintTemplateDefs != nullptr)
+			{
+				ProjectInfo.StarterContent = SelectedTemplate->BlueprintTemplateDefs->StarterContent;
+			}
+		}
+	}
 
 	const TArray<ETemplateSetting>& HiddenSettings = SelectedTemplate->HiddenSettings;
 
