Blank blueprint projects now open the Minimal_Default level.

This was mistakenly removed by a previous refactoring to the created config when creating a project from a template.

diff --git a/Engine/Source/Editor/GameProjectGeneration/Private/GameProjectUtils.cpp b/Engine/Source/Editor/GameProjectGeneration/Private/GameProjectUtils.cpp
--- a/Engine/Source/Editor/GameProjectGeneration/Private/GameProjectUtils.cpp
+++ b/Engine/Source/Editor/GameProjectGeneration/Private/GameProjectUtils.cpp
@@ -224,6 +224,12 @@ namespace
 				TEXT("EditorStartupMap"),
 				DefaultMap,
 				true /* ShouldReplaceExistingValue */);
+
+			ConfigValues.Emplace(TEXT("DefaultEngine.ini"),
+				TEXT("/Script/EngineSettings.GameMapsSettings"),
+				TEXT("GameDefaultMap"),
+				DefaultMap,
+				true /* ShouldReplaceExistingValue */);
 		}
 	}
 } // namespace <>
