Fixing merge issues introduced in CL 10384788 - removed GetHardwareConfigString declaration and GetDefaultMapConfigString function.

diff --git a/Engine/Source/Editor/GameProjectGeneration/Private/GameProjectUtils.cpp b/Engine/Source/Editor/GameProjectGeneration/Private/GameProjectUtils.cpp
--- a/Engine/Source/Editor/GameProjectGeneration/Private/GameProjectUtils.cpp
+++ b/Engine/Source/Editor/GameProjectGeneration/Private/GameProjectUtils.cpp
@@ -851,26 +851,6 @@ bool GameProjectUtils::IsStarterContentAvailableForProject(const FProjectInforma
 	return IFileManager::Get().FileExists(*StarterContentPackFilename);
 }
 
-FString GameProjectUtils::GetDefaultMapConfigString(const FProjectInformation& InProjectInfo)
-{
-	FString DefaultMapConfig;
-
-	if (InProjectInfo.bIsBlankTemplate &&
-		InProjectInfo.bCopyStarterContent &&
-		GameProjectUtils::IsStarterContentAvailableForProject(InProjectInfo))
-	{
-		const FString StarterContentName = GameProjectUtils::GetStarterContentName(InProjectInfo);
-
-		DefaultMapConfig += LINE_TERMINATOR;
-		DefaultMapConfig += TEXT("[/Script/EngineSettings.GameMapsSettings]") LINE_TERMINATOR;
-
-		DefaultMapConfig += FString::Printf(TEXT("EditorStartupMap=/Game/%s/Maps/Minimal_Default") LINE_TERMINATOR, *StarterContentName);
-		DefaultMapConfig += FString::Printf(TEXT("GameDefaultMap=/Game/%s/Maps/Minimal_Default") LINE_TERMINATOR, *StarterContentName);
-	}
-
-	return MoveTemp(DefaultMapConfig);
-}
-
 bool GameProjectUtils::CreateProject(const FProjectInformation& InProjectInfo, FText& OutFailReason, FText& OutFailLog, TArray<FString>* OutCreatedFiles)
 {
 	if ( !IsValidProjectFileForCreation(InProjectInfo.ProjectFilename, OutFailReason) )
@@ -2092,8 +2072,6 @@ bool GameProjectUtils::GenerateConfigFiles(const FProjectInformation& InProjectI
 
 		if (InProjectInfo.bCopyStarterContent)
 		{
-			FileContents += GetDefaultMapConfigString(InProjectInfo);
-
 			if (GameProjectUtils::IsStarterContentAvailableForNewProjects())
 			{
 				if (InProjectInfo.TargetedHardware == EHardwareClass::Mobile)
diff --git a/Engine/Source/Editor/GameProjectGeneration/Public/GameProjectUtils.h b/Engine/Source/Editor/GameProjectGeneration/Public/GameProjectUtils.h
--- a/Engine/Source/Editor/GameProjectGeneration/Public/GameProjectUtils.h
+++ b/Engine/Source/Editor/GameProjectGeneration/Public/GameProjectUtils.h
@@ -285,17 +285,13 @@ public:
 	static const TCHAR* GetDefaultBuildSettingsVersion();
 
 private:
-
+	
+	/** Add hardware-specific config values such as the target platform and RHI. */
 	static void AddHardwareConfigValues(const FProjectInformation& InProjectInfo, TArray<FTemplateConfigValue>& ConfigValues);
+	
 	/** Get the name of the starter content pack to use for the given project. */
 	static FString GetStarterContentName(const FProjectInformation& InProjectInfo);
 
-	/** Get the config string for the hardware target for the given project. */
-	static FString GetHardwareConfigString(const FProjectInformation& InProjectInfo);
-
-	/** Get the config string to add for the default maps for the given project. */
-	static FString GetDefaultMapConfigString(const FProjectInformation& InProjectInfo);
-
 	/** Generates a new project without using a template project */
 	static bool GenerateProjectFromScratch(const FProjectInformation& InProjectInfo, FText& OutFailReason, FText& OutFailLog);
 
