When adding the very first C++ class to an otherwise source-less project, change the internal UBT target name to the default editor target name for the project

diff --git a/Engine/Source/Editor/GameProjectGeneration/Private/GameProjectUtils.cpp b/Engine/Source/Editor/GameProjectGeneration/Private/GameProjectUtils.cpp
--- a/Engine/Source/Editor/GameProjectGeneration/Private/GameProjectUtils.cpp
+++ b/Engine/Source/Editor/GameProjectGeneration/Private/GameProjectUtils.cpp
@@ -4025,6 +4025,11 @@ GameProjectUtils::EAddCodeToProjectResult GameProjectUtils::AddCodeToProject_Int
 		const FString GameModuleName = FApp::GetProjectName();
 		check(ModuleInfo.ModuleName == GameModuleName);
 
+		// Because this project previously didn't have any code, the UBT target name will just be UE4Editor. Now that we've
+		// added some code, the target name will be changed to match the editor target for the new source. 
+		FString NewUBTTargetName = GameModuleName + TEXT("Editor");
+		FPlatformMisc::SetUBTTargetName(*NewUBTTargetName);
+
 		IHotReloadInterface& HotReloadSupport = FModuleManager::LoadModuleChecked<IHotReloadInterface>("HotReload");
 		if (!HotReloadSupport.RecompileModule(*GameModuleName, *GWarn, ERecompileModuleFlags::ReloadAfterRecompile | ERecompileModuleFlags::ForceCodeProject))
 		{
