Optimize directory iterations on startup
  - Plugin discovery
    - 560ms to 50ms (warm cache)
    - 3.3s to 365ms (cold cache)
  - Source file discovery
    - 102ms to 23ms (warm cache)
    - 253ms to 41ms (cold cache)

  - Cold cache = Cleared standby pages, mostly equivalent to a reboot
  - Warm cache = 2nd run of 2 consecutive runs

#rnx

diff --git a/Engine/Source/Editor/GameProjectGeneration/Private/GameProjectUtils.cpp b/Engine/Source/Editor/GameProjectGeneration/Private/GameProjectUtils.cpp
--- a/Engine/Source/Editor/GameProjectGeneration/Private/GameProjectUtils.cpp
+++ b/Engine/Source/Editor/GameProjectGeneration/Private/GameProjectUtils.cpp
@@ -3747,44 +3747,90 @@ FString GameProjectUtils::GetDefaultProjectTemplateFilename()
 	return TEXT("");
 }
 
-void FindCodeFiles(const TCHAR* BaseDirectory, TArray<FString>& FileNames, int32 MaxNumFileNames)
+void FindCodeFiles(const TCHAR* BaseDirectory, TFunctionRef<bool(const TCHAR* FileName)> SourceFoundThreadSafeCallback)
 {
+	// Can't use IterateDirectoryRecursive here as we rely on some directory filtering without aborting whole enumeration
 	struct FDirectoryVisitor : public IPlatformFile::FDirectoryVisitor
 	{
-		TArray<FString>& FileNames;
-		int32 MaxNumFileNames;
+		TFunctionRef<bool(const TCHAR* FileName)> SourceFoundThreadSafeCallback;
+		FRWLock          DirectoriesLock;
+		TArray<FString>& Directories;
 
-		FDirectoryVisitor(TArray<FString>& InFileNames, int32 InMaxNumFileNames)
-			: FileNames(InFileNames)
-			, MaxNumFileNames(InMaxNumFileNames)
+		FDirectoryVisitor(TFunctionRef<bool(const TCHAR* FileName)> InSourceFoundThreadSafeCallback, TArray<FString>& InDirectories)
+			: IPlatformFile::FDirectoryVisitor(EDirectoryVisitorFlags::ThreadSafe)
+			, SourceFoundThreadSafeCallback(InSourceFoundThreadSafeCallback)
+			, Directories(InDirectories)
 		{
 		}
 
 		virtual bool Visit(const TCHAR* FilenameOrDirectory, bool bIsDirectory) override
 		{
-			if(bIsDirectory)
+			if (bIsDirectory)
 			{
 				FString CleanDirectoryName(FPaths::GetCleanFilename(FilenameOrDirectory));
-				if(!CleanDirectoryName.StartsWith(TEXT(".")))
+				// Skip directories like (i.e. .git) when finding code files to improve performance.
+				if (!CleanDirectoryName.StartsWith(TEXT(".")))
 				{
-					FindCodeFiles(FilenameOrDirectory, FileNames, MaxNumFileNames);
+					FString Directory(FilenameOrDirectory);
+					FRWScopeLock ScopeLock(DirectoriesLock, SLT_Write);
+					Directories.Emplace(MoveTemp(Directory));
 				}
 			}
 			else
 			{
-				FString FileName(FilenameOrDirectory);
-				if(FileName.EndsWith(TEXT(".h")) || FileName.EndsWith(".cpp"))
+				FStringView FileName(FilenameOrDirectory);
+				if (FileName.EndsWith(TEXT(".h")) || FileName.EndsWith(TEXT(".cpp")))
 				{
-					FileNames.Add(FileName);
+					return SourceFoundThreadSafeCallback(FilenameOrDirectory);
 				}
 			}
-			return (FileNames.Num() < MaxNumFileNames);
+			return true;
 		}
 	};
 
-	// Enumerate the contents of the current directory
-	FDirectoryVisitor Visitor(FileNames, MaxNumFileNames);
-	FPlatformFileManager::Get().GetPlatformFile().IterateDirectory(BaseDirectory, Visitor);
+	TArray<FString> DirectoriesToVisitNext;
+	DirectoriesToVisitNext.Add(BaseDirectory);
+	
+	TAtomic<bool> bResult(true);
+	FDirectoryVisitor Visitor(SourceFoundThreadSafeCallback, DirectoriesToVisitNext);
+	IPlatformFile& PlatformFile = FPlatformFileManager::Get().GetPlatformFile();
+	while (bResult && DirectoriesToVisitNext.Num() > 0)
+	{
+		TArray<FString> DirectoriesToVisit = MoveTemp(DirectoriesToVisitNext);
+		ParallelFor(
+			DirectoriesToVisit.Num(),
+			[&PlatformFile , &DirectoriesToVisit, &Visitor, &bResult](int32 Index)
+			{
+				if (bResult && !PlatformFile.IterateDirectory(*DirectoriesToVisit[Index], Visitor))
+				{
+					bResult = false;
+				}
+			},
+			EParallelForFlags::Unbalanced
+		);
+	}
+}
+
+void FindCodeFiles(const TCHAR* BaseDirectory, TArray<FString>& FileNames, int32 MaxNumFileNames)
+{
+	FRWLock FileNamesLock;
+	auto SourceFoundThreadSafeCallback =
+		[&FileNamesLock, &FileNames, MaxNumFileNames](const TCHAR* FileName) -> bool
+		{
+			// Prepare the string out of the lock
+			FString SourceFile(FileName);
+
+			FRWScopeLock ScopeLock(FileNamesLock, SLT_Write);
+			// Make sure we don't go over the requested MaxNumFileNames even
+			// when multiple threads are reaching this at the same time
+			if (FileNames.Num() < MaxNumFileNames)
+			{
+				FileNames.Emplace(MoveTemp(SourceFile));
+			}
+			return FileNames.Num() < MaxNumFileNames;
+		};
+
+	FindCodeFiles(BaseDirectory, SourceFoundThreadSafeCallback);
 }
 
 void GameProjectUtils::GetProjectCodeFilenames(TArray<FString>& OutProjectCodeFilenames)
@@ -3801,15 +3847,21 @@ int32 GameProjectUtils::GetProjectCodeFileCount()
 
 void GameProjectUtils::GetProjectSourceDirectoryInfo(int32& OutNumCodeFiles, int64& OutDirectorySize)
 {
-	TArray<FString> Filenames;
-	GetProjectCodeFilenames(Filenames);
-	OutNumCodeFiles = Filenames.Num();
+	TRACE_CPUPROFILER_EVENT_SCOPE(GameProjectUtils::GetProjectSourceDirectoryInfo)
 
-	OutDirectorySize = 0;
-	for (const auto& filename : Filenames)
-	{
-		OutDirectorySize += IFileManager::Get().FileSize(*filename);
-	}
+	TAtomic<int32> NumCodeFiles(0);
+	TAtomic<int64> DirectorySize(0);
+	auto SourceFoundThreadSafeCallback =
+		[&DirectorySize, &NumCodeFiles](const TCHAR* FileName) -> bool
+		{
+			NumCodeFiles++;
+			DirectorySize += IFileManager::Get().FileSize(FileName);
+			return true;
+		};
+
+	FindCodeFiles(*FPaths::GameSourceDir(), SourceFoundThreadSafeCallback);
+	OutDirectorySize = DirectorySize;
+	OutNumCodeFiles = NumCodeFiles;
 }
 
 bool GameProjectUtils::ProjectHasCodeFiles()
diff --git a/Engine/Source/Runtime/Core/Private/GenericPlatform/GenericPlatformFile.cpp b/Engine/Source/Runtime/Core/Private/GenericPlatform/GenericPlatformFile.cpp
--- a/Engine/Source/Runtime/Core/Private/GenericPlatform/GenericPlatformFile.cpp
+++ b/Engine/Source/Runtime/Core/Private/GenericPlatform/GenericPlatformFile.cpp
@@ -13,6 +13,8 @@
 
 #include "Async/AsyncFileHandle.h"
 #include "Async/MappedFileHandle.h"
+#include "Async/ParallelFor.h"
+#include "Misc/ScopeRWLock.h"
 
 class FGenericBaseRequest;
 class FGenericAsyncReadFileHandle;
@@ -532,25 +534,52 @@ bool IPlatformFile::IterateDirectoryRecursively(const TCHAR* Directory, FDirecto
 	class FRecurse : public FDirectoryVisitor
 	{
 	public:
-		IPlatformFile&		PlatformFile;
-		FDirectoryVisitor&	Visitor;
-		FRecurse(IPlatformFile&	InPlatformFile, FDirectoryVisitor& InVisitor)
-			: PlatformFile(InPlatformFile)
+		IPlatformFile&      PlatformFile;
+		FDirectoryVisitor&  Visitor;
+		FRWLock             DirectoriesLock;
+		TArray<FString>&    Directories;
+		FRecurse(IPlatformFile&	InPlatformFile, FDirectoryVisitor& InVisitor, TArray<FString>& InDirectories)
+			: FDirectoryVisitor(InVisitor.DirectoryVisitorFlags)
+			, PlatformFile(InPlatformFile)
 			, Visitor(InVisitor)
+			, Directories(InDirectories)
 		{
 		}
 		virtual bool Visit(const TCHAR* FilenameOrDirectory, bool bIsDirectory) override
 		{
-			bool Result = Visitor.Visit(FilenameOrDirectory, bIsDirectory);
-			if (Result && bIsDirectory)
+			bool bResult = Visitor.Visit(FilenameOrDirectory, bIsDirectory);
+			if (bResult && bIsDirectory)
 			{
-				Result = PlatformFile.IterateDirectory(FilenameOrDirectory, *this);
+				FString Directory(FilenameOrDirectory);
+				FRWScopeLock ScopeLock(DirectoriesLock, SLT_Write);
+				Directories.Emplace(MoveTemp(Directory));
 			}
-			return Result;
+			return bResult;
 		}
 	};
-	FRecurse Recurse(*this, Visitor);
-	return IterateDirectory(Directory, Recurse);
+
+	TArray<FString> DirectoriesToVisitNext;
+	DirectoriesToVisitNext.Add(Directory);
+
+	TAtomic<bool> bResult(true);
+	FRecurse Recurse(*this, Visitor, DirectoriesToVisitNext);
+	while (bResult && DirectoriesToVisitNext.Num() > 0)
+	{
+		TArray<FString> DirectoriesToVisit = MoveTemp(DirectoriesToVisitNext);
+		ParallelFor(
+			DirectoriesToVisit.Num(),
+			[this, &DirectoriesToVisit, &Recurse, &bResult](int32 Index)
+			{
+				if (bResult && !IterateDirectory(*DirectoriesToVisit[Index], Recurse))
+				{
+					bResult = false;
+				}
+			},
+			Visitor.IsThreadSafe() ? EParallelForFlags::Unbalanced : EParallelForFlags::ForceSingleThread
+		);
+	}
+
+	return bResult;
 }
 
 bool IPlatformFile::IterateDirectoryStatRecursively(const TCHAR* Directory, FDirectoryStatVisitor& Visitor)
@@ -567,12 +596,12 @@ bool IPlatformFile::IterateDirectoryStatRecursively(const TCHAR* Directory, FDir
 		}
 		virtual bool Visit(const TCHAR* FilenameOrDirectory, const FFileStatData& StatData) override
 		{
-			bool Result = Visitor.Visit(FilenameOrDirectory, StatData);
-			if (Result && StatData.bIsDirectory)
+			bool bResult = Visitor.Visit(FilenameOrDirectory, StatData);
+			if (bResult && StatData.bIsDirectory)
 			{
-				Result = PlatformFile.IterateDirectoryStat(FilenameOrDirectory, *this);
+				bResult = PlatformFile.IterateDirectoryStat(FilenameOrDirectory, *this);
 			}
-			return Result;
+			return bResult;
 		}
 	};
 	FStatRecurse Recurse(*this, Visitor);
@@ -594,12 +623,14 @@ bool IPlatformFile::IterateDirectoryStatRecursively(const TCHAR* Directory, FDir
 class FFindFilesVisitor : public IPlatformFile::FDirectoryVisitor
 {
 public:
-	IPlatformFile&		PlatformFile;
-	TArray<FString>&	FoundFiles;
-	const TCHAR*		FileExtension;
-	int32				FileExtensionLen;
+	IPlatformFile&   PlatformFile;
+	FRWLock          FoundFilesLock;
+	TArray<FString>& FoundFiles;
+	const TCHAR*     FileExtension;
+	int32            FileExtensionLen;
 	FFindFilesVisitor(IPlatformFile& InPlatformFile, TArray<FString>& InFoundFiles, const TCHAR* InFileExtension)
-		: PlatformFile(InPlatformFile)
+		: IPlatformFile::FDirectoryVisitor(EDirectoryVisitorFlags::ThreadSafe)
+		, PlatformFile(InPlatformFile)
 		, FoundFiles(InFoundFiles)
 		, FileExtension(InFileExtension)
 		, FileExtensionLen(InFileExtension ? FCString::Strlen(InFileExtension) : 0)
@@ -618,8 +649,10 @@ public:
 					return true;
 				}
 			}
-				
-			FoundFiles.Emplace(FString(FilenameOrDirectory));
+			
+			FString FileName(FilenameOrDirectory);
+			FRWScopeLock ScopeLock(FoundFilesLock, SLT_Write);
+			FoundFiles.Emplace(MoveTemp(FileName));
 		}
 		return true;
 	}
diff --git a/Engine/Source/Runtime/Core/Public/GenericPlatform/GenericPlatformFile.h b/Engine/Source/Runtime/Core/Public/GenericPlatform/GenericPlatformFile.h
--- a/Engine/Source/Runtime/Core/Public/GenericPlatform/GenericPlatformFile.h
+++ b/Engine/Source/Runtime/Core/Public/GenericPlatform/GenericPlatformFile.h
@@ -64,6 +64,17 @@ enum class EPlatformFileWrite : uint8
 
 ENUM_CLASS_FLAGS(EPlatformFileWrite);
 
+/**
+ * Enum for the DirectoryVisitor flags
+ */
+enum class EDirectoryVisitorFlags : uint8
+{
+	None = 0x0,
+	ThreadSafe = 0x01	// should be set when the Visit function can be called from multiple threads at once.
+};
+
+ENUM_CLASS_FLAGS(EDirectoryVisitorFlags);
+
 /** 
  * File handle interface. 
 **/
@@ -337,6 +348,11 @@ public:
 	class FDirectoryVisitor
 	{
 	public:
+		FDirectoryVisitor(EDirectoryVisitorFlags InDirectoryVisitorFlags = EDirectoryVisitorFlags::None)
+			: DirectoryVisitorFlags(InDirectoryVisitorFlags)
+		{
+		}
+
 		virtual ~FDirectoryVisitor() { }
 
 		/** 
@@ -346,6 +362,14 @@ public:
 		 * @return							true if the iteration should continue.
 		**/
 		virtual bool Visit(const TCHAR* FilenameOrDirectory, bool bIsDirectory) = 0;
+
+		/** True if the Visit function can be called from multiple threads at once. **/
+		FORCEINLINE bool IsThreadSafe() const
+		{
+			return (DirectoryVisitorFlags & EDirectoryVisitorFlags::ThreadSafe) != EDirectoryVisitorFlags::None;
+		}
+
+		EDirectoryVisitorFlags DirectoryVisitorFlags;
 	};
 
 	/** File and directory visitor function that takes only the name */
diff --git a/Engine/Source/Runtime/Projects/Private/PluginManager.cpp b/Engine/Source/Runtime/Projects/Private/PluginManager.cpp
--- a/Engine/Source/Runtime/Projects/Private/PluginManager.cpp
+++ b/Engine/Source/Runtime/Projects/Private/PluginManager.cpp
@@ -20,6 +20,7 @@
 #include "ProjectManager.h"
 #include "PluginManifest.h"
 #include "HAL/PlatformTime.h"
+#include "Async/ParallelFor.h"
 #if READ_TARGET_ENABLED_PLUGINS_FROM_RECEIPT
 #include "IDesktopPlatform.h"
 #include "DesktopPlatformModule.h"
@@ -447,19 +448,40 @@ void FPluginManager::ReadPluginsInDirectory(const FString& PluginsDirectory, con
 		TArray<FString> FileNames;
 		FindPluginsInDirectory(PluginsDirectory, FileNames);
 
-		for(const FString& FileName: FileNames)
+		struct FLoadContext
 		{
 			FPluginDescriptor Descriptor;
 			FText FailureReason;
-			if ( Descriptor.Load(FileName, FailureReason) )
+			bool bResult = false;
+		};
+
+		TArray<FLoadContext> Contexts;
+		Contexts.SetNum(FileNames.Num());
+
+		ParallelFor(
+			FileNames.Num(),
+			[&Contexts, &FileNames](int32 Index)
+			{
+				FLoadContext& Context = Contexts[Index];
+				Context.bResult = Context.Descriptor.Load(FileNames[Index], Context.FailureReason);
+			},
+			EParallelForFlags::Unbalanced
+		);
+
+		for (int32 Index = 0, Num = FileNames.Num(); Index < Num; ++Index)
+		{
+			const FString& FileName = FileNames[Index];
+			FLoadContext& Context = Contexts[Index];
+			
+			if (Context.bResult)
 			{
-				CreatePluginObject(FileName, Descriptor, Type, Plugins, ChildPlugins);
+				CreatePluginObject(FileName, Context.Descriptor, Type, Plugins, ChildPlugins);
 			}
 			else
 			{
 				// NOTE: Even though loading of this plugin failed, we'll keep processing other plugins
 				FString FullPath = FPaths::ConvertRelativePathToFull(FileName);
-				FText FailureMessage = FText::Format(LOCTEXT("FailureFormat", "{0} ({1})"), FailureReason, FText::FromString(FullPath));
+				FText FailureMessage = FText::Format(LOCTEXT("FailureFormat", "{0} ({1})"), Context.FailureReason, FText::FromString(FullPath));
 				FText DialogTitle = LOCTEXT("PluginFailureTitle", "Failed to load Plugin");
 				UE_LOG(LogPluginManager, Error, TEXT("%s"), *FailureMessage.ToString());
 				FMessageDialog::Open(EAppMsgType::Ok, FailureMessage, &DialogTitle);
@@ -488,10 +510,10 @@ void FPluginManager::FindPluginManifestsInDirectory(const FString& PluginManifes
 		{
 			if (!bIsDirectory)
 			{
-				FString FileName(FileNameOrDirectory);
-				if (FileName.EndsWith(".upluginmanifest"))
+				FStringView FileName(FileNameOrDirectory);
+				if (FileName.EndsWith(TEXT(".upluginmanifest")))
 				{
-					ManifestFileNames.Add(FileName);
+					ManifestFileNames.Emplace(FileName);
 				}
 			}
 			return true;
