Fix missing include in CL 12670961

#rnx

diff --git a/Engine/Source/Editor/GameProjectGeneration/Private/GameProjectUtils.cpp b/Engine/Source/Editor/GameProjectGeneration/Private/GameProjectUtils.cpp
--- a/Engine/Source/Editor/GameProjectGeneration/Private/GameProjectUtils.cpp
+++ b/Engine/Source/Editor/GameProjectGeneration/Private/GameProjectUtils.cpp
@@ -77,6 +77,7 @@
 #include "PlatformInfo.h"
 #include "Blueprint/BlueprintSupport.h"
 #include "Settings/ProjectPackagingSettings.h"
+#include "Async/ParallelFor.h"
 
 #if WITH_LIVE_CODING
 #include "ILiveCodingModule.h"
