Class Template Registry

diff --git a/Engine/Plugins/Runtime/AudioModulation/Source/AudioModulationEditor/AudioModulationEditor.build.cs b/Engine/Plugins/Runtime/AudioModulation/Source/AudioModulationEditor/AudioModulationEditor.build.cs
--- a/Engine/Plugins/Runtime/AudioModulation/Source/AudioModulationEditor/AudioModulationEditor.build.cs
+++ b/Engine/Plugins/Runtime/AudioModulation/Source/AudioModulationEditor/AudioModulationEditor.build.cs
@@ -12,6 +12,7 @@ public class AudioModulationEditor : ModuleRules
 				"CoreUObject",
 				"Engine",
 				"EditorFramework",
+				"GameProjectGeneration",
 				"UnrealEd",
 				"PropertyEditor",
 				"SequenceRecorder",
diff --git a/Engine/Source/Editor/AudioEditor/AudioEditor.Build.cs b/Engine/Source/Editor/AudioEditor/AudioEditor.Build.cs
--- a/Engine/Source/Editor/AudioEditor/AudioEditor.Build.cs
+++ b/Engine/Source/Editor/AudioEditor/AudioEditor.Build.cs
@@ -22,6 +22,7 @@ public class AudioEditor : ModuleRules
 			{
 				"AudioMixer",
 				"EditorSubsystem",
+				"GameProjectGeneration",
 				"ToolMenus",
 				"UMG",
 				"UMGEditor",
diff --git a/Engine/Source/Editor/AudioEditor/Private/SoundClassTemplates.cpp b/Engine/Source/Editor/AudioEditor/Private/SoundClassTemplates.cpp
new file mode 100644
--- /dev/null
+++ b/Engine/Source/Editor/AudioEditor/Private/SoundClassTemplates.cpp
@@ -0,0 +1,20 @@
+// Copyright Epic Games, Inc. All Rights Reserved.
+#include "SoundClassTemplates.h"
+
+USoundEffectSourcePresetClassTemplate::USoundEffectSourcePresetClassTemplate(const FObjectInitializer& ObjectInitializer)
+	: UClassTemplate(ObjectInitializer)
+{
+	SetGeneratedBaseClass(USoundEffectSourcePreset::StaticClass());
+}
+
+USoundEffectSubmixPresetClassTemplate::USoundEffectSubmixPresetClassTemplate(const FObjectInitializer& ObjectInitializer)
+	: UClassTemplate(ObjectInitializer)
+{
+	SetGeneratedBaseClass(USoundEffectSubmixPreset::StaticClass());
+}
+
+USynthComponentClassTemplate::USynthComponentClassTemplate(const FObjectInitializer& ObjectInitializer)
+	: UClassTemplate(ObjectInitializer)
+{
+	SetGeneratedBaseClass(USynthComponent::StaticClass());
+}
diff --git a/Engine/Source/Editor/AudioEditor/Private/SoundClassTemplates.h b/Engine/Source/Editor/AudioEditor/Private/SoundClassTemplates.h
new file mode 100644
--- /dev/null
+++ b/Engine/Source/Editor/AudioEditor/Private/SoundClassTemplates.h
@@ -0,0 +1,47 @@
+// Copyright Epic Games, Inc. All Rights Reserved.
+#pragma once
+
+#include "ClassTemplateEditorSubsystem.h"
+#include "Components/SynthComponent.h"
+#include "CoreMinimal.h"
+#include "Sound/SoundEffectSource.h"
+#include "Sound/SoundEffectSubmix.h"
+
+#include "SoundClassTemplates.generated.h"
+
+
+UCLASS()
+class USoundEffectSourcePresetClassTemplate : public UClassTemplate
+{
+	GENERATED_UCLASS_BODY()
+
+public:
+	FString GetFilename() const override
+	{
+		return TEXT("SoundEffectSourceClass");
+	}
+};
+
+UCLASS()
+class USoundEffectSubmixPresetClassTemplate : public UClassTemplate
+{
+	GENERATED_UCLASS_BODY()
+
+public:
+	FString GetFilename() const override
+	{
+		return TEXT("SoundEffectSubmixClass");
+	}
+};
+
+UCLASS()
+class USynthComponentClassTemplate : public UClassTemplate
+{
+	GENERATED_UCLASS_BODY()
+
+public:
+	FString GetFilename() const override
+	{
+		return TEXT("SynthComponentClass");
+	}
+};
\ No newline at end of file
diff --git a/Engine/Source/Editor/GameProjectGeneration/GameProjectGeneration.Build.cs b/Engine/Source/Editor/GameProjectGeneration/GameProjectGeneration.Build.cs
--- a/Engine/Source/Editor/GameProjectGeneration/GameProjectGeneration.Build.cs
+++ b/Engine/Source/Editor/GameProjectGeneration/GameProjectGeneration.Build.cs
@@ -4,7 +4,7 @@ using UnrealBuildTool;
 
 public class GameProjectGeneration : ModuleRules
 {
-    public GameProjectGeneration(ReadOnlyTargetRules Target) : base(Target)
+	public GameProjectGeneration(ReadOnlyTargetRules Target) : base(Target)
 	{
 		PublicDependencyModuleNames.AddRange(
 			new string[] {
@@ -12,8 +12,8 @@ public class GameProjectGeneration : ModuleRules
 			}
 		);
 
-        PrivateIncludePathModuleNames.AddRange(
-            new string[] {
+		PrivateIncludePathModuleNames.AddRange(
+			new string[] {
 				"AssetRegistry",
 				"ContentBrowser",
 				"DesktopPlatform",
@@ -21,7 +21,7 @@ public class GameProjectGeneration : ModuleRules
 				"MainFrame",
 				"AddContentDialog",
 				"HardwareTargeting",
-            }
+			}
 		);
 
 		PrivateDependencyModuleNames.AddRange(
@@ -43,23 +43,24 @@ public class GameProjectGeneration : ModuleRules
 				"SourceControl",
  				"TargetPlatform",
 				"EditorFramework",
+				"EditorSubsystem",
 				"UnrealEd",
 				"DesktopPlatform",
 				"LauncherPlatform",
 				"HardwareTargeting",
 				"AddContentDialog",
-                "AudioMixer",
+				"AudioMixer",
 				"AudioMixerCore"
 			}
 		);
 
-        DynamicallyLoadedModuleNames.AddRange(
-            new string[] {
+		DynamicallyLoadedModuleNames.AddRange(
+			new string[] {
 				"AssetRegistry",
 				"ContentBrowser",
-                "Documentation",
-                "MainFrame",
-            }
+				"Documentation",
+				"MainFrame",
+			}
 		);
 
 		if(Target.bWithLiveCoding)
diff --git a/Engine/Source/Editor/GameProjectGeneration/Private/ClassTemplateEditorSubsystem.cpp b/Engine/Source/Editor/GameProjectGeneration/Private/ClassTemplateEditorSubsystem.cpp
new file mode 100644
--- /dev/null
+++ b/Engine/Source/Editor/GameProjectGeneration/Private/ClassTemplateEditorSubsystem.cpp
@@ -0,0 +1,199 @@
+// Copyright Epic Games, Inc. All Rights Reserved.
+#include "ClassTemplateEditorSubsystem.h"
+
+#include "Interfaces/IPluginManager.h"
+#include "Internationalization/Text.h"
+#include "Misc/CoreDelegates.h"
+#include "Misc/FileHelper.h"
+#include "Subsystems/SubsystemCollection.h"
+#include "UObject/UObjectGlobals.h"
+
+#define LOCTEXT_NAMESPACE "ClassTemplateSubsystem"
+
+
+FString UClassTemplate::GetDirectory() const
+{
+	return UClassTemplateEditorSubsystem::GetEngineTemplateDirectory();
+}
+
+FString UClassTemplate::GetFilename() const
+{
+	if (GeneratedBaseClass)
+	{
+		return GeneratedBaseClass->GetName();
+	}
+
+	return *GetName();
+}
+
+FString UClassTemplate::GetHeaderFilename() const
+{
+	static const FString Extension = TEXT(".h.template");
+	return GetFilename() + Extension;
+}
+
+FString UClassTemplate::GetSourceFilename() const
+{
+	static const FString Extension = TEXT(".cpp.template");
+	return GetFilename() + Extension;
+}
+
+bool UClassTemplate::ReadHeader(FString& OutHeaderFileText, FText& OutFailReason) const
+{
+	const FString Dir = GetDirectory();
+	const FString FilePath = Dir / GetHeaderFilename();
+
+	if (!FFileHelper::LoadFileToString(OutHeaderFileText, *FilePath))
+	{
+		FFormatNamedArguments Args;
+		Args.Add(TEXT("FilePath"), FText::FromString(FilePath));
+		OutFailReason = FText::Format(LOCTEXT("ReadClassTemplateFailure_FailedToReadTemplateHeaderFile", "Failed to read template source file \"{FilePath}\""), Args);
+		return false;
+	}
+
+	return true;
+}
+
+bool UClassTemplate::ReadSource(FString& OutSourceFileText, FText& OutFailReason) const
+{
+	const FString Dir = GetDirectory();
+	const FString FilePath = Dir / GetSourceFilename();
+
+	if (!FFileHelper::LoadFileToString(OutSourceFileText, *FilePath))
+	{
+		FFormatNamedArguments Args;
+		Args.Add(TEXT("FilePath"), FText::FromString(FilePath));
+		OutFailReason = FText::Format(LOCTEXT("ReadClassTemplateFailure_FailedToReadTemplateSourceFile", "Failed to read template source file \"{FilePath}\""), Args);
+		return false;
+	}
+
+	return true;
+}
+
+void UClassTemplate::BeginDestroy()
+{
+	if (GEditor)
+	{
+		if (UClassTemplateEditorSubsystem* TemplateSubsystem = GEditor->GetEditorSubsystem<UClassTemplateEditorSubsystem>())
+		{
+			TemplateSubsystem->Unregister(this);
+		}
+	}
+
+	Super::BeginDestroy();
+}
+
+const UClass* UClassTemplate::GetGeneratedBaseClass() const
+{
+	return GeneratedBaseClass;
+}
+
+void UClassTemplate::SetGeneratedBaseClass(UClass* InClass)
+{
+	GeneratedBaseClass = InClass;
+
+	const UClass* Class = GetClass();
+	check(Class);
+
+	const bool bIsAbstract = Class->HasAnyClassFlags(CLASS_Abstract);
+	if (GEditor && !bIsAbstract)
+	{
+		if (UClassTemplateEditorSubsystem* TemplateSubsystem = GEditor->GetEditorSubsystem<UClassTemplateEditorSubsystem>())
+		{
+			if (TemplateSubsystem->ContainsClassTemplate(GeneratedBaseClass))
+			{
+				TemplateSubsystem->Unregister(this);
+			}
+
+			TemplateSubsystem->Register(this);
+		}
+	}
+}
+
+FString UPluginClassTemplate::GetDirectory() const
+{
+	TSharedPtr<IPlugin> Plugin = IPluginManager::Get().FindPlugin(PluginName);
+	if (Plugin.IsValid())
+	{
+		return Plugin->GetContentDir() / TEXT("Templates");
+	}
+
+	return UClassTemplate::GetDirectory();
+}
+
+UClassTemplateEditorSubsystem::UClassTemplateEditorSubsystem()
+	: UEditorSubsystem()
+{
+	FCoreDelegates::OnPostEngineInit.AddUObject(this, &UClassTemplateEditorSubsystem::RegisterTemplates);
+}
+
+void UClassTemplateEditorSubsystem::Deinitialize()
+{
+	TemplateRegistry.Reset();
+	FCoreDelegates::OnPostEngineInit.RemoveAll(this);
+}
+
+void UClassTemplateEditorSubsystem::Initialize(FSubsystemCollectionBase& Collection)
+{
+	RegisterTemplates();
+}
+
+void UClassTemplateEditorSubsystem::RegisterTemplates()
+{
+	for (TObjectIterator<UClass> It; It; ++It)
+	{
+		if (It->IsChildOf(UClassTemplate::StaticClass()) && !It->HasAnyClassFlags(CLASS_Abstract))
+		{
+			const UClassTemplate* ClassTemplate = Cast<const UClassTemplate>((*It)->GetDefaultObject());
+			Register(ClassTemplate);
+		}
+	}
+}
+
+void UClassTemplateEditorSubsystem::Register(const UClassTemplate* InClassTemplate)
+{
+	if (!ensure(InClassTemplate))
+	{
+		return;
+	}
+
+	TWeakObjectPtr<const UClassTemplate> ClassTemplate = InClassTemplate;
+	TWeakObjectPtr<const UClass> GeneratedBaseClass = ClassTemplate->GetGeneratedBaseClass();
+	if (ensure(GeneratedBaseClass.IsValid()))
+	{
+		const bool bIsRegistered = TemplateRegistry.Contains(GeneratedBaseClass);
+		if (!bIsRegistered)
+		{
+			TemplateRegistry.Add(GeneratedBaseClass, ClassTemplate);
+		}
+	}
+}
+
+bool UClassTemplateEditorSubsystem::Unregister(const UClassTemplate* InClassTemplate)
+{
+	if (!ensure(InClassTemplate))
+	{
+		return false;
+	}
+
+	TWeakObjectPtr<const UClass> Class = InClassTemplate->GetGeneratedBaseClass();
+	return TemplateRegistry.Remove(Class) > 0;
+}
+
+bool UClassTemplateEditorSubsystem::ContainsClassTemplate(const UClass* InClass) const
+{
+	const TWeakObjectPtr<const UClass> ClassPtr = MakeWeakObjectPtr(InClass);
+	return TemplateRegistry.Contains(ClassPtr);
+}
+
+const UClassTemplate* UClassTemplateEditorSubsystem::FindClassTemplate(const UClass* InClass) const
+{
+	const TWeakObjectPtr<const UClass> ClassPtr = MakeWeakObjectPtr(InClass);
+	return TemplateRegistry.FindRef(ClassPtr).Get();
+}
+
+FString UClassTemplateEditorSubsystem::GetEngineTemplateDirectory()
+{
+	return FPaths::EngineContentDir() / TEXT("Editor") / TEXT("Templates");
+}
+#undef LOCTEXT_NAMESPACE
\ No newline at end of file
diff --git a/Engine/Source/Editor/GameProjectGeneration/Private/GameProjectUtils.cpp b/Engine/Source/Editor/GameProjectGeneration/Private/GameProjectUtils.cpp
--- a/Engine/Source/Editor/GameProjectGeneration/Private/GameProjectUtils.cpp
+++ b/Engine/Source/Editor/GameProjectGeneration/Private/GameProjectUtils.cpp
@@ -19,6 +19,7 @@
 #include "Widgets/DeclarativeSyntaxSupport.h"
 #include "Widgets/SWindow.h"
 #include "Framework/Application/SlateApplication.h"
+#include "ClassTemplateEditorSubsystem.h"
 #include "Components/ActorComponent.h"
 #include "Components/SceneComponent.h"
 #include "GameFramework/Actor.h"
@@ -70,9 +71,6 @@
 #include "IAudioExtensionPlugin.h"
 #include "AudioPluginUtilities.h"
 #include "Sound/AudioSettings.h"
-#include "Sound/SoundEffectSubmix.h"
-#include "Sound/SoundEffectSource.h"
-#include "Components/SynthComponent.h"
 
 #include "PlatformInfo.h"
 #include "Blueprint/BlueprintSupport.h"
@@ -520,18 +518,6 @@ FString FNewClassInfo::GetHeaderTemplateFilename() const
 				{
 					return TEXT("CharacterClass.h.template");
 				}
-				else if (BaseClass == USoundEffectSourcePreset::StaticClass())
-				{
-					return TEXT("SoundEffectSourceClass.h.template");
-				}
-				else if (BaseClass == USoundEffectSubmixPreset::StaticClass())
-				{
-					return TEXT("SoundEffectSubmixClass.h.template");
-				}
-				else if (BaseClass == USynthComponent::StaticClass())
-				{
-					return TEXT("SynthComponentClass.h.template");
-				}
 			}
 			// Some other non-actor, non-component UObject class
 			return TEXT( "UObjectClass.h.template" );
@@ -578,18 +564,6 @@ FString FNewClassInfo::GetSourceTemplateFilename() const
 				{
 					return TEXT("CharacterClass.cpp.template");
 				}
-				else if (BaseClass == USoundEffectSubmixPreset::StaticClass())
-				{
-					return TEXT("SoundEffectSubmixClass.cpp.template");
-				}
-				else if (BaseClass == USoundEffectSourcePreset::StaticClass())
-				{
-					return TEXT("SoundEffectSourceClass.cpp.template");
-				}
-				else if (BaseClass == USynthComponent::StaticClass())
-				{
-					return TEXT("SynthComponentClass.cpp.template");
-				}
 			}
 			// Some other non-actor, non-component UObject class
 			return TEXT( "UObjectClass.cpp.template" );
@@ -3000,7 +2974,7 @@ const TCHAR* GameProjectUtils::GetDefaultBuildSettingsVersion()
 
 bool GameProjectUtils::ReadTemplateFile(const FString& TemplateFileName, FString& OutFileContents, FText& OutFailReason)
 {
-	const FString FullFileName = FPaths::EngineContentDir() / TEXT("Editor") / TEXT("Templates") / TemplateFileName;
+	const FString FullFileName = UClassTemplateEditorSubsystem::GetEngineTemplateDirectory() / TemplateFileName;
 	if ( FFileHelper::LoadFileToString(OutFileContents, *FullFileName) )
 	{
 		return true;
@@ -3180,7 +3154,24 @@ bool GenerateConstructorDeclaration(FString& Out, const FString& PrefixedClassNa
 bool GameProjectUtils::GenerateClassHeaderFile(const FString& NewHeaderFileName, const FString UnPrefixedClassName, const FNewClassInfo ParentClassInfo, const TArray<FString>& ClassSpecifierList, const FString& ClassProperties, const FString& ClassFunctionDeclarations, FString& OutSyncLocation, const FModuleContextInfo& ModuleInfo, bool bDeclareConstructor, FText& OutFailReason)
 {
 	FString Template;
-	if ( !ReadTemplateFile(ParentClassInfo.GetHeaderTemplateFilename(), Template, OutFailReason) )
+	bool bTemplateFound = false;
+	if (GEditor)
+	{
+		if (UClassTemplateEditorSubsystem* TemplateSubsystem = GEditor->GetEditorSubsystem<UClassTemplateEditorSubsystem>())
+		{
+			const UClass* BaseClass = ParentClassInfo.BaseClass;
+			if (const UClassTemplate* ClassTemplate = TemplateSubsystem->FindClassTemplate(ParentClassInfo.BaseClass))
+			{
+				bTemplateFound = ClassTemplate->ReadHeader(Template, OutFailReason);
+				if (!bTemplateFound)
+				{
+					return false;
+				}
+			}
+		}
+	}
+
+	if (!bTemplateFound && !ReadTemplateFile(ParentClassInfo.GetHeaderTemplateFilename(), Template, OutFailReason))
 	{
 		return false;
 	}
@@ -3225,6 +3216,7 @@ bool GameProjectUtils::GenerateClassHeaderFile(const FString& NewHeaderFileName,
 	FinalOutput = FinalOutput.Replace(TEXT("%UCLASS_SPECIFIER_LIST%"), *MakeCommaDelimitedList(ClassSpecifierList, false), ESearchCase::CaseSensitive);
 	FinalOutput = FinalOutput.Replace(TEXT("%PREFIXED_CLASS_NAME%"), *PrefixedClassName, ESearchCase::CaseSensitive);
 	FinalOutput = FinalOutput.Replace(TEXT("%PREFIXED_BASE_CLASS_NAME%"), *PrefixedBaseClassName, ESearchCase::CaseSensitive);
+	FinalOutput = FinalOutput.Replace(TEXT("%MODULE_NAME%"), *ModuleInfo.ModuleName, ESearchCase::CaseSensitive);
 
 	// Special case where where the wildcard starts with a tab and ends with a new line
 	const bool bLeadingTab = true;
@@ -3327,7 +3319,24 @@ static bool IsUsingOldStylePch(FString BaseDir)
 bool GameProjectUtils::GenerateClassCPPFile(const FString& NewCPPFileName, const FString UnPrefixedClassName, const FNewClassInfo ParentClassInfo, const TArray<FString>& AdditionalIncludes, const TArray<FString>& PropertyOverrides, const FString& AdditionalMemberDefinitions, FString& OutSyncLocation, const FModuleContextInfo& ModuleInfo, FText& OutFailReason)
 {
 	FString Template;
-	if ( !ReadTemplateFile(ParentClassInfo.GetSourceTemplateFilename(), Template, OutFailReason) )
+	bool bTemplateFound = false;
+	if (GEditor)
+	{
+		if (UClassTemplateEditorSubsystem* TemplateSubsystem = GEditor->GetEditorSubsystem<UClassTemplateEditorSubsystem>())
+		{
+			const UClass* BaseClass = ParentClassInfo.BaseClass;
+			if (const UClassTemplate* ClassTemplate = TemplateSubsystem->FindClassTemplate(ParentClassInfo.BaseClass))
+			{
+				bTemplateFound = ClassTemplate->ReadSource(Template, OutFailReason);
+				if (!bTemplateFound)
+				{
+					return false;
+				}
+			}
+		}
+	}
+
+	if (!bTemplateFound && !ReadTemplateFile(ParentClassInfo.GetSourceTemplateFilename(), Template, OutFailReason))
 	{
 		return false;
 	}
diff --git a/Engine/Source/Editor/GameProjectGeneration/Public/ClassTemplateEditorSubsystem.h b/Engine/Source/Editor/GameProjectGeneration/Public/ClassTemplateEditorSubsystem.h
new file mode 100644
--- /dev/null
+++ b/Engine/Source/Editor/GameProjectGeneration/Public/ClassTemplateEditorSubsystem.h
@@ -0,0 +1,109 @@
+// Copyright Epic Games, Inc. All Rights Reserved.
+#pragma once
+
+#include "CoreMinimal.h"
+#include "EditorSubsystem.h"
+#include "UObject/Class.h"
+#include "UObject/WeakObjectPtrTemplates.h"
+
+#include "ClassTemplateEditorSubsystem.generated.h"
+
+
+// Forward Declarations
+class FSubsystemCollectionBase;
+class FText;
+
+
+UCLASS(Abstract)
+class GAMEPROJECTGENERATION_API UClassTemplate : public UObject
+{
+	GENERATED_BODY()
+
+public:
+	virtual void BeginDestroy() override;
+
+	// Returns the directory containing the text file template for
+	// the given base generated class.
+	virtual FString GetDirectory() const;
+
+	// Reads the header template text from disk.  If failure to read from
+	// disk, returns false and provides text reason.
+	bool ReadHeader(FString& OutHeaderFileText, FText& OutFailReason) const;
+
+	// Reads the source template text from disk.  If failure to read from
+	// disk, returns false and provides text reason.
+	bool ReadSource(FString& OutSourceFileText, FText& OutFailReason) const;
+
+	const UClass* GetGeneratedBaseClass() const;
+
+protected:
+	// Sets the generated base class associated with the given template
+	void SetGeneratedBaseClass(UClass* InClass);
+
+	// Returns the filename associated with the provided class template
+	// without an extension.  Defaults to class name.
+	virtual FString GetFilename() const;
+
+	// Returns full header filename including '.h.template' extension
+	FString GetHeaderFilename() const;
+
+	// Returns full sourcefilename including '.cpp.template' extension
+	FString GetSourceFilename() const;
+
+private:
+	// Base UClass of which template class corresponds.
+	UPROPERTY(Transient)
+	UClass* GeneratedBaseClass;
+};
+
+UCLASS(Abstract)
+class GAMEPROJECTGENERATION_API UPluginClassTemplate : public UClassTemplate
+{
+	GENERATED_BODY()
+
+public:
+	// Returns the directory containing the text file template for
+	// the given base generated class.
+	virtual FString GetDirectory() const override;
+
+protected:
+	UPROPERTY(Transient)
+	FString PluginName;
+};
+
+UCLASS()
+class GAMEPROJECTGENERATION_API UClassTemplateEditorSubsystem : public UEditorSubsystem
+{
+	GENERATED_BODY()
+
+private:
+	using FTemplateRegistry = TMap<TWeakObjectPtr<const UClass>, TWeakObjectPtr<const UClassTemplate>>;
+	FTemplateRegistry TemplateRegistry;
+
+public:
+	UClassTemplateEditorSubsystem();
+
+	virtual void Initialize(FSubsystemCollectionBase& Collection) override;
+	virtual void Deinitialize() override;
+
+	// Registers all currently loaded template classes with the internal registry.
+	void RegisterTemplates();
+
+	// Returns path to the directory containing all engine class templates.
+	static FString GetEngineTemplateDirectory();
+
+	// Returns whether or not class has registered template
+	bool ContainsClassTemplate(const UClass* InClass) const;
+
+	// Returns class template if one is registered.
+	const UClassTemplate* FindClassTemplate(const UClass* InClass) const;
+
+	friend class UClassTemplate;
+
+private:
+	// Registers a template class with the subsystem
+	void Register(const UClassTemplate* InClassTemplate);
+
+	// Unregisters a template class with the subsystem
+	bool Unregister(const UClassTemplate* InClassTemplate);
+};
