Enable world partition by default for newly created projects.

diff --git a/Engine/Source/Editor/GameProjectGeneration/Private/GameProjectUtils.cpp b/Engine/Source/Editor/GameProjectGeneration/Private/GameProjectUtils.cpp
--- a/Engine/Source/Editor/GameProjectGeneration/Private/GameProjectUtils.cpp
+++ b/Engine/Source/Editor/GameProjectGeneration/Private/GameProjectUtils.cpp
@@ -2109,6 +2109,12 @@ bool GameProjectUtils::GenerateConfigFiles(const FProjectInformation& InProjectI
 			}
 		}
 
+		FileContents += LINE_TERMINATOR;
+		FileContents += TEXT("[/Script/WorldPartitionEditor.WorldPartitionEditorSettings]") LINE_TERMINATOR;
+		FileContents += TEXT("bEnableWorldPartition=True") LINE_TERMINATOR;
+		FileContents += TEXT("bEnableConversionPrompt=True") LINE_TERMINATOR;
+		FileContents += TEXT("CommandletClass=Class'/Script/UnrealEd.WorldPartitionConvertCommandlet'") LINE_TERMINATOR;
+
 		if (WriteOutputFile(DefaultEngineIniFilename, FileContents, OutFailReason))
 		{
 			OutCreatedFiles.Add(DefaultEngineIniFilename);
