Fixed string usage issues found by UE_NODISCARD

#rnx

diff --git a/Engine/Plugins/Experimental/BackChannel/Source/BackChannel/Private/Protocol/BackChannelOSCConnection.cpp b/Engine/Plugins/Experimental/BackChannel/Source/BackChannel/Private/Protocol/BackChannelOSCConnection.cpp
--- a/Engine/Plugins/Experimental/BackChannel/Source/BackChannel/Private/Protocol/BackChannelOSCConnection.cpp
+++ b/Engine/Plugins/Experimental/BackChannel/Source/BackChannel/Private/Protocol/BackChannelOSCConnection.cpp
@@ -362,7 +362,7 @@ int32 FBackChannelOSCConnection::GetMessageLimitForPath(const TCHAR* InPath)
 
 	if (Path.EndsWith(TEXT("*")))
 	{
-		Path.LeftChop(1);
+		Path.LeftChopInline(1);
 	}
 
 	// todo - search for vest match, not first match
diff --git a/Engine/Source/Editor/GameProjectGeneration/Private/SNewProjectWizard.cpp b/Engine/Source/Editor/GameProjectGeneration/Private/SNewProjectWizard.cpp
--- a/Engine/Source/Editor/GameProjectGeneration/Private/SNewProjectWizard.cpp
+++ b/Engine/Source/Editor/GameProjectGeneration/Private/SNewProjectWizard.cpp
@@ -985,9 +985,9 @@ void SNewProjectWizard::SetDefaultProjectLocation()
 		IFileManager::Get().MakeDirectory(*DefaultProjectFilePath, true);
 	}
 
-	if ( !DefaultProjectFilePath.IsEmpty() && DefaultProjectFilePath.Right(1) == TEXT("/") )
+	if ( DefaultProjectFilePath.EndsWith(TEXT("/")) )
 	{
-		DefaultProjectFilePath.LeftChop(1);
+		DefaultProjectFilePath.LeftChopInline(1);
 	}
 
 	FPaths::NormalizeFilename(DefaultProjectFilePath);
diff --git a/Engine/Source/Runtime/Core/Private/Misc/Paths.cpp b/Engine/Source/Runtime/Core/Private/Misc/Paths.cpp
--- a/Engine/Source/Runtime/Core/Private/Misc/Paths.cpp
+++ b/Engine/Source/Runtime/Core/Private/Misc/Paths.cpp
@@ -997,7 +997,7 @@ bool FPaths::IsDrive(const FString& InPath)
 		if (CheckPath.Len() > 0)
 		{
 			// Replace any remaining '\\' instances with '\'
-			CheckPath.Replace(TEXT("\\\\"), TEXT("\\"), ESearchCase::CaseSensitive);
+			CheckPath.ReplaceInline(TEXT("\\\\"), TEXT("\\"), ESearchCase::CaseSensitive);
 
 			int32 CheckCount = 0;
 			int32 SlashIndex = CheckPath.Find(TEXT("\\"), ESearchCase::CaseSensitive);
diff --git a/Engine/Source/Runtime/CoreUObject/Private/UObject/PropertyBaseObject.cpp b/Engine/Source/Runtime/CoreUObject/Private/UObject/PropertyBaseObject.cpp
--- a/Engine/Source/Runtime/CoreUObject/Private/UObject/PropertyBaseObject.cpp
+++ b/Engine/Source/Runtime/CoreUObject/Private/UObject/PropertyBaseObject.cpp
@@ -225,8 +225,6 @@ void FObjectPropertyBase::ExportTextItem( FString& ValueStr, const void* Propert
 
 	if (0 != (PortFlags & PPF_ExportCpp))
 	{
-		FString::Printf(TEXT("%s%s*"), PropertyClass->GetPrefixCPP(), *PropertyClass->GetName());
-
 		ValueStr += Temp
 			? FString::Printf(TEXT("LoadObject<%s%s>(nullptr, TEXT(\"%s\"))")
 				, PropertyClass->GetPrefixCPP()
