Project dialog reskin

diff --git a/Engine/Source/Editor/DetailCustomizations/Private/HardwareTargetingSettingsDetails.cpp b/Engine/Source/Editor/DetailCustomizations/Private/HardwareTargetingSettingsDetails.cpp
--- a/Engine/Source/Editor/DetailCustomizations/Private/HardwareTargetingSettingsDetails.cpp
+++ b/Engine/Source/Editor/DetailCustomizations/Private/HardwareTargetingSettingsDetails.cpp
@@ -275,19 +275,18 @@ void FHardwareTargetingSettingsDetails::CustomizeDetails(IDetailLayoutBuilder& D
 		DetailBuilder.HideProperty(PropertyName);
 
 		TSharedRef<IPropertyHandle> Property = DetailBuilder.GetProperty(PropertyName);
-		auto SetPropertyValue = [](EHardwareClass::Type NewValue, TSharedRef<IPropertyHandle> InProperty){
+		auto SetPropertyValue = [](EHardwareClass NewValue, TSharedRef<IPropertyHandle> InProperty){
 			InProperty->SetValue(uint8(NewValue));
 		};
 		auto GetPropertyValue = [](TSharedRef<IPropertyHandle> InProperty){
 			uint8 Value = 0;
 			InProperty->GetValue(Value);
-			return EHardwareClass::Type(Value);
+			return EHardwareClass(Value);
 		};
 		
 		HardwareClassCombo = HardwareTargeting.MakeHardwareClassTargetCombo(
 			FOnHardwareClassChanged::CreateStatic(SetPropertyValue, Property),
-			TAttribute<EHardwareClass::Type>::Create(TAttribute<EHardwareClass::Type>::FGetter::CreateStatic(GetPropertyValue, Property)),
-			Orient_Horizontal
+			TAttribute<EHardwareClass>::Create(TAttribute<EHardwareClass>::FGetter::CreateStatic(GetPropertyValue, Property))
 		);
 	}
 
@@ -297,42 +296,41 @@ void FHardwareTargetingSettingsDetails::CustomizeDetails(IDetailLayoutBuilder& D
 		DetailBuilder.HideProperty(PropertyName);
 
 		TSharedRef<IPropertyHandle> Property = DetailBuilder.GetProperty(PropertyName);
-		auto SetPropertyValue = [](EGraphicsPreset::Type NewValue, TSharedRef<IPropertyHandle> InProperty){
+		auto SetPropertyValue = [](EGraphicsPreset NewValue, TSharedRef<IPropertyHandle> InProperty){
 			InProperty->SetValue(uint8(NewValue));
 		};
 		auto GetPropertyValue = [](TSharedRef<IPropertyHandle> InProperty){
 			uint8 Value = 0;
 			InProperty->GetValue(Value);
-			return EGraphicsPreset::Type(Value);
+			return EGraphicsPreset(Value);
 		};
 		GraphicsPresetCombo = HardwareTargeting.MakeGraphicsPresetTargetCombo(
 			FOnGraphicsPresetChanged::CreateStatic(SetPropertyValue, Property),
-			TAttribute<EGraphicsPreset::Type>::Create(TAttribute<EGraphicsPreset::Type>::FGetter::CreateStatic(GetPropertyValue, Property)),
-			Orient_Horizontal
+			TAttribute<EGraphicsPreset>::Create(TAttribute<EGraphicsPreset>::FGetter::CreateStatic(GetPropertyValue, Property))
 		);
 	}
 
-	HardwareTargetingCategory.AddCustomRow(LOCTEXT("HardwareTargetingOption", "Targeted Hardware:"))
+	HardwareTargetingCategory.AddCustomRow(LOCTEXT("HardwareTargetingOption", "Targeted Hardware"))
 	.NameContent()
 	[
 		SNew(STextBlock)
-		.Text(LOCTEXT("OptimizeProjectFor", "Optimize project settings for:"))
+		.Text(LOCTEXT("OptimizeProjectFor", "Optimize project settings for"))
 		.Font(DetailBuilder.GetDetailFont())
 	]
 	.ValueContent()
 	.MaxDesiredWidth(0)
 	[
 		SNew(SHorizontalBox)
-
 		+ SHorizontalBox::Slot()
 		.Padding(FMargin(10.f, 0.f))
+		.VAlign(VAlign_Center)
 		.AutoWidth()
 		[
 			HardwareClassCombo.ToSharedRef()
 		]
-
 		+ SHorizontalBox::Slot()
 		.Padding(FMargin(0.f, 0.f, 10.f, 0.f))
+		.VAlign(VAlign_Center)
 		.AutoWidth()
 		[
 			GraphicsPresetCombo.ToSharedRef()
diff --git a/Engine/Source/Editor/EditorStyle/Private/StarshipStyle.cpp b/Engine/Source/Editor/EditorStyle/Private/StarshipStyle.cpp
--- a/Engine/Source/Editor/EditorStyle/Private/StarshipStyle.cpp
+++ b/Engine/Source/Editor/EditorStyle/Private/StarshipStyle.cpp
@@ -422,16 +422,15 @@ void FStarshipEditorStyle::FStyle::SetupGeneralStyles()
 
 	// Common icons
 	{
-		Set( "Icons.Refresh", new IMAGE_BRUSH( "Icons/icon_Refresh_16x", Icon16x16 ) );
 		Set( "Icons.Contact", new IMAGE_BRUSH( "Icons/icon_mail_16x", Icon16x16 ) );
 
 		Set( "Icons.DirtyBadge", new IMAGE_BRUSH_SVG("Starship/Common/DirtyBadge", Icon12x12));
 	}
 
 	Set( "WarningStripe", new IMAGE_BRUSH( "Common/WarningStripe", FVector2D(20,6), FLinearColor::White, ESlateBrushTileType::Horizontal ) );
-
-	Set("RoundedWarning", new FSlateRoundedBoxBrush(FStyleColors::Warning, 4.0f));
-	Set("RoundedError", new FSlateRoundedBoxBrush(FStyleColors::Error, 4.0f));
+	
+	Set("RoundedWarning", new FSlateRoundedBoxBrush(FStyleColors::Transparent, 4.0f, FStyleColors::Warning, 1.0f));
+	Set("RoundedError", new FSlateRoundedBoxBrush(FStyleColors::Transparent, 4.0f, FStyleColors::Error, 1.0f));
 
 	Set( "Button.Disabled", new BOX_BRUSH( "Common/Button_Disabled", 8.0f/32.0f ) );
 	
@@ -1686,24 +1685,6 @@ void FStarshipEditorStyle::FStyle::SetupGeneralStyles()
 		Set( "GameProjectDialog.FolderIconOpen", new IMAGE_BRUSH( "Icons/FolderOpen", FVector2D(18, 16) ) );
 		Set( "GameProjectDialog.ProjectFileIcon", new IMAGE_BRUSH( "Icons/doc_16x", FVector2D(18, 16) ) );
 
-		Set( "GameProjectDialog.IncludeStarterContent", new IMAGE_BRUSH( "/GameProjectDialog/IncludeStarterContent", FVector2D(64, 64) ) );
-		Set( "GameProjectDialog.NoStarterContent", new IMAGE_BRUSH( "/GameProjectDialog/NoStarterContent", FVector2D(64, 64) ) );
-		
-		Set("GameProjectDialog.XRDisabled", new IMAGE_BRUSH("Icons/AssetIcons/SkeletalMesh_64x", Icon64x64));
-		Set("GameProjectDialog.XREnabled", new IMAGE_BRUSH("Icons/AssetIcons/SkeletalMesh_64x", Icon64x64));
-
-		Set("GameProjectDialog.RaytracingDisabled", new IMAGE_BRUSH("Icons/AssetIcons/Procedural_Foliage_Spawner_64x", Icon64x64));
-		Set("GameProjectDialog.RaytracingEnabled", new IMAGE_BRUSH("Icons/AssetIcons/Procedural_Foliage_Spawner_64x", Icon64x64));
-
-		Set( "FilePath.FolderButton",
-			FButtonStyle(HoverHintOnly)
-			.SetNormal(  BOX_BRUSH( "Common/ButtonHoverHint", FMargin(4/16.0f), FLinearColor(0.5f,0.5f,0.5f) ) )
-			.SetHovered( BOX_BRUSH( "Common/ButtonHoverHint", FMargin(4/16.0f), FLinearColor(0.6f,0.6f,0.6f) ) )
-			.SetPressed( BOX_BRUSH( "Common/ButtonHoverHint", FMargin(4/16.0f), FLinearColor(0.65f,0.65f,0.65f) ) )
-			.SetNormalPadding( FMargin(0,0,0,1) )
-			.SetPressedPadding( FMargin(0,1,0,0) )
-			);
-		Set( "FilePath.GroupIndicator", new BOX_BRUSH( "GameProjectDialog/filepath_group_indicator", FMargin(4.0f/16.f) ) );
 	}
 
 	// NewClassDialog
@@ -6557,31 +6538,43 @@ void FStarshipEditorStyle::FStyle::SetupToolkitStyles()
 #if WITH_EDITOR
 	// Project Browser
 	{
-		Set("ProjectBrowser.Tab.Text", FTextBlockStyle(NormalText)
-			.SetFont(DEFAULT_FONT("Regular", 24))
-			.SetShadowOffset(FVector2D(0, 1)));
-
-		Set("ProjectBrowser.Toolbar.Text", FTextBlockStyle(NormalText)
-			.SetFont(DEFAULT_FONT("Bold", 12))
-			.SetShadowColorAndOpacity(FLinearColor(0, 0, 0, 0.5f))
-			.SetShadowOffset(FVector2D(0, 1)));
-
 		Set("ProjectBrowser.VersionOverlayText", FTextBlockStyle(NormalText)
 			.SetFont(DEFAULT_FONT("Bold", 14))
 			.SetShadowOffset(FVector2D(0, 1)));
 
 
-		Set("ProjectBrowser.Background", new BOX_BRUSH("Common/FlatColorSquare", FVector2D(1.0f, 1.0f), FMargin(0), FLinearColor(FColor(0xff404040))));
-		Set("ProjectBrowser.Tab.ActiveBackground", new BOX_BRUSH("Common/FlatColorSquare", FVector2D(1.0f, 1.0f), FMargin(0), FLinearColor(FColor(0xff404040))));
-		Set("ProjectBrowser.Tab.Background", new BOX_BRUSH("Common/FlatColorSquare", FVector2D(1.0f, 1.0f), FMargin(0), FLinearColor(FColor(0xff272727))));
-		Set("ProjectBrowser.Tab.ActiveHighlight", new BOX_BRUSH("Common/FlatColorSquare", FVector2D(1.0f, 1.0f), FMargin(0), SelectionColor));
-		Set("ProjectBrowser.Tab.Highlight", new BOX_BRUSH("Common/FlatColorSquare", FVector2D(1.0f, 1.0f), FMargin(0), SelectionColor_Inactive));
-		Set("ProjectBrowser.Tab.PressedHighlight", new BOX_BRUSH("Common/FlatColorSquare", FVector2D(1.0f, 1.0f), FMargin(0), SelectionColor_Pressed));
+		const FTableRowStyle ProjectBrowserTableRowStyle = FTableRowStyle()
+			.SetEvenRowBackgroundBrush(FSlateNoResource())
+			.SetEvenRowBackgroundHoveredBrush(FSlateNoResource())
+			.SetOddRowBackgroundBrush(FSlateNoResource())
+			.SetOddRowBackgroundHoveredBrush(FSlateNoResource())
+			.SetSelectorFocusedBrush(FSlateNoResource())
+			.SetActiveBrush(FSlateNoResource())
+			.SetActiveHoveredBrush(FSlateNoResource())
+			.SetInactiveBrush(FSlateNoResource())
+			.SetInactiveHoveredBrush(FSlateNoResource())
+			.SetActiveHighlightedBrush(FSlateNoResource())
+			.SetInactiveHighlightedBrush(FSlateNoResource())
+			.SetTextColor(FStyleColors::Foreground)
+			.SetSelectedTextColor(FStyleColors::ForegroundInverted);
+
+		Set("ProjectBrowser.TableRow", ProjectBrowserTableRowStyle);
+
+		Set("ProjectBrowser.MajorCategoryViewBorder", new FSlateRoundedBoxBrush(FStyleColors::Recessed, 4.0f));
+
 
 		Set("ProjectBrowser.TileViewTooltip.ToolTipBorder", new FSlateColorBrush(FLinearColor::Black));
 		Set("ProjectBrowser.TileViewTooltip.NonContentBorder", new BOX_BRUSH("/Docking/TabContentArea", FMargin(4 / 16.0f)));
 		Set("ProjectBrowser.TileViewTooltip.ContentBorder", new BOX_BRUSH("Common/GroupBorder", FMargin(4.0f / 16.0f)));
 		Set("ProjectBrowser.TileViewTooltip.NameFont", DEFAULT_FONT("Regular", 12));
+
+		Set("ProjectBrowser.ProjectTile.Font", DEFAULT_FONT("Regular", 9));
+		Set("ProjectBrowser.ProjectTile.ThumbnailAreaBackground", new FSlateRoundedBoxBrush(COLOR("#474747FF"), FVector4(4.0f,4.0f,0.0f,0.0f)));
+		Set("ProjectBrowser.ProjectTile.NameAreaBackground", new FSlateRoundedBoxBrush(EStyleColor::Header, FVector4(0.0f, 0.0f, 4.0f, 4.0f)));
+		Set("ProjectBrowser.ProjectTile.DropShadow", new BOX_BRUSH("Starship/ContentBrowser/drop-shadow", FMargin(4.0f / 64.0f)));
+		Set("ProjectBrowser.ProjectTile.SelectedBorder", new FSlateRoundedBoxBrush(FStyleColors::Transparent, 4.0f, FStyleColors::Primary, 1.0f));
+		Set("ProjectBrowser.ProjectTile.SelectedHoverBorder", new FSlateRoundedBoxBrush(FStyleColors::Transparent, 4.0f, FStyleColors::PrimaryHover, 1.0f));
+		Set("ProjectBrowser.ProjectTile.HoverBorder", new FSlateRoundedBoxBrush(FStyleColors::Transparent, 4.0f, FStyleColors::Hover, 1.0f));
 	}
 
 	// Toolkit Display
diff --git a/Engine/Source/Editor/GameProjectGeneration/Classes/TemplateProjectDefs.h b/Engine/Source/Editor/GameProjectGeneration/Classes/TemplateProjectDefs.h
--- a/Engine/Source/Editor/GameProjectGeneration/Classes/TemplateProjectDefs.h
+++ b/Engine/Source/Editor/GameProjectGeneration/Classes/TemplateProjectDefs.h
@@ -129,6 +129,10 @@ public:
 	UPROPERTY(config)
 	bool bIsBlank;
 
+	/** Is there a rendered thumbnail that should be treated as the project template icon. If this is true the thumbnail takes up the full tile size rather than a 64x64 icon */
+	UPROPERTY(config)
+	bool bThumbnailAsIcon;
+
 	/* Optional list of settings to hide. If none are specified, then all settings are shown. */
 	UPROPERTY(config)
 	TArray<ETemplateSetting> HiddenSettings;
@@ -193,13 +197,9 @@ struct FTemplateCategoryDef
 	UPROPERTY()
 	TArray<FLocalizedTemplateString> LocalizedDescriptions;
 
-	/** Reference to an icon to display for this category. Should be around 128x128. */
+	/** Reference to an icon to display for this category. Should be around 300x100. */
 	UPROPERTY()
 	FString Icon;
-
-	/** Is this a major top-level category? Major categories are displayed as full rows, eg. the Game category.*/
-	UPROPERTY()
-	bool IsMajorCategory;
 };
 
 UCLASS(config=TemplateCategories)
diff --git a/Engine/Source/Editor/GameProjectGeneration/GameProjectGeneration.Build.cs b/Engine/Source/Editor/GameProjectGeneration/GameProjectGeneration.Build.cs
--- a/Engine/Source/Editor/GameProjectGeneration/GameProjectGeneration.Build.cs
+++ b/Engine/Source/Editor/GameProjectGeneration/GameProjectGeneration.Build.cs
@@ -40,6 +40,7 @@ public class GameProjectGeneration : ModuleRules
 				"Slate",
 				"SlateCore",
 				"EditorStyle",
+				"EditorWidgets",
 				"SourceControl",
  				"TargetPlatform",
 				"EditorFramework",
diff --git a/Engine/Source/Editor/GameProjectGeneration/Private/GameProjectGenerationModule.cpp b/Engine/Source/Editor/GameProjectGeneration/Private/GameProjectGenerationModule.cpp
--- a/Engine/Source/Editor/GameProjectGeneration/Private/GameProjectGenerationModule.cpp
+++ b/Engine/Source/Editor/GameProjectGeneration/Private/GameProjectGenerationModule.cpp
@@ -7,13 +7,14 @@
 #include "EditorStyleSet.h"
 #include "GameProjectGenerationLog.h"
 #include "GameProjectUtils.h"
-#include "SGameProjectDialog.h"
 #include "SNewClassDialog.h"
 #include "TemplateCategory.h"
 #include "HAL/FileManager.h"
 #include "Interfaces/IPluginManager.h"
 #include "TemplateProjectDefs.h"
 #include "Internationalization/Culture.h"
+#include "SProjectDialog.h"
+
 
 IMPLEMENT_MODULE( FGameProjectGenerationModule, GameProjectGeneration );
 DEFINE_LOG_CATEGORY(LogGameProjectGeneration);
@@ -34,18 +35,18 @@ TSharedRef<class SWidget> FGameProjectGenerationModule::CreateGameProjectDialog(
 {
 	ensure(bAllowProjectOpening || bAllowProjectCreate);
 
-	SGameProjectDialog::EMode Mode = SGameProjectDialog::EMode::Both;
-	
+	EProjectDialogModeMode Mode = EProjectDialogModeMode::Hybrid;
+
 	if (bAllowProjectOpening && !bAllowProjectCreate)
 	{
-		Mode = SGameProjectDialog::EMode::Open;
+		Mode = EProjectDialogModeMode::OpenProject;
 	}
 	else if (bAllowProjectCreate && !bAllowProjectOpening)
 	{
-		Mode = SGameProjectDialog::EMode::New;
+		Mode = EProjectDialogModeMode::NewProject;
 	}
 
-	return SNew(SGameProjectDialog, Mode);
+	return SNew(SProjectDialog, Mode);
 }
 
 
@@ -192,9 +193,13 @@ void FGameProjectGenerationModule::LoadTemplateCategories()
 				TemplateCategory->Description = FLocalizedTemplateString::GetLocalizedText(Category.LocalizedDescriptions);
 				
 				const FName BrushName(*Category.Icon);
-				TemplateCategory->Icon = new FSlateDynamicImageBrush(BrushName, FVector2D(128, 128));
+				
+				FSlateDynamicImageBrush* Brush = new FSlateDynamicImageBrush(BrushName, FVector2D(300, 100));
+				Brush->OutlineSettings.CornerRadii = FVector4(4, 4, 4, 4);
+				Brush->OutlineSettings.RoundingType = ESlateBrushRoundingType::FixedRadius;
+				Brush->DrawAs = ESlateBrushDrawType::RoundedBox;
 
-				TemplateCategory->IsMajor = Category.IsMajorCategory;
+				TemplateCategory->Icon = Brush;
 			}
 		}
 	}
diff --git a/Engine/Source/Editor/GameProjectGeneration/Private/GameProjectUtils.cpp b/Engine/Source/Editor/GameProjectGeneration/Private/GameProjectUtils.cpp
--- a/Engine/Source/Editor/GameProjectGeneration/Private/GameProjectUtils.cpp
+++ b/Engine/Source/Editor/GameProjectGeneration/Private/GameProjectUtils.cpp
@@ -913,19 +913,19 @@ bool GameProjectUtils::CreateProject(const FProjectInformation& InProjectInfo, F
 
 		if (InProjectInfo.TargetedHardware.IsSet())
 		{
-			UEnum* HardwareClassEnum = StaticEnum<EHardwareClass::Type>();
+			UEnum* HardwareClassEnum = StaticEnum<EHardwareClass>();
 			if (HardwareClassEnum != nullptr)
 			{
-				EventAttributes.Add(FAnalyticsEventAttribute(TEXT("HardwareClass"), HardwareClassEnum->GetNameStringByValue(InProjectInfo.TargetedHardware.GetValue())));
+				EventAttributes.Add(FAnalyticsEventAttribute(TEXT("HardwareClass"), HardwareClassEnum->GetNameStringByValue(static_cast<int32>(InProjectInfo.TargetedHardware.GetValue()))));
 			}
 		}
 
 		if (InProjectInfo.DefaultGraphicsPerformance.IsSet())
 		{
-			UEnum* GraphicsPresetEnum = StaticEnum<EGraphicsPreset::Type>();
+			UEnum* GraphicsPresetEnum = StaticEnum<EGraphicsPreset>();
 			if (GraphicsPresetEnum != nullptr)
 			{
-				EventAttributes.Add(FAnalyticsEventAttribute(TEXT("GraphicsPreset"), GraphicsPresetEnum->GetNameStringByValue(InProjectInfo.DefaultGraphicsPerformance.GetValue())));
+				EventAttributes.Add(FAnalyticsEventAttribute(TEXT("GraphicsPreset"), GraphicsPresetEnum->GetNameStringByValue(static_cast<int32>(InProjectInfo.DefaultGraphicsPerformance.GetValue()))));
 			}
 		}
 
@@ -2049,7 +2049,7 @@ void GameProjectUtils::AddHardwareConfigValues(const FProjectInformation& InProj
 {
 	if (InProjectInfo.TargetedHardware.IsSet())
 	{
-		UEnum* HardwareClassEnum = StaticEnum<EHardwareClass::Type>();
+		UEnum* HardwareClassEnum = StaticEnum<EHardwareClass>();
 		if (HardwareClassEnum != nullptr)
 		{
 			FString TargetHardwareString;
@@ -2068,7 +2068,7 @@ void GameProjectUtils::AddHardwareConfigValues(const FProjectInformation& InProj
 
 	if (InProjectInfo.DefaultGraphicsPerformance.IsSet())
 	{
-		UEnum* GraphicsPresetEnum = StaticEnum<EGraphicsPreset::Type>();
+		UEnum* GraphicsPresetEnum = StaticEnum<EGraphicsPreset>();
 		if (GraphicsPresetEnum != nullptr)
 		{
 			FString GraphicsPresetString;
diff --git a/Engine/Source/Editor/GameProjectGeneration/Private/SGameProjectDialog.cpp b/Engine/Source/Editor/GameProjectGeneration/Private/SGameProjectDialog.cpp
deleted file mode 100644
diff --git a/Engine/Source/Editor/GameProjectGeneration/Private/SGameProjectDialog.h b/Engine/Source/Editor/GameProjectGeneration/Private/SGameProjectDialog.h
deleted file mode 100644
diff --git a/Engine/Source/Editor/GameProjectGeneration/Private/SGetSuggestedIDEWidget.cpp b/Engine/Source/Editor/GameProjectGeneration/Private/SGetSuggestedIDEWidget.cpp
--- a/Engine/Source/Editor/GameProjectGeneration/Private/SGetSuggestedIDEWidget.cpp
+++ b/Engine/Source/Editor/GameProjectGeneration/Private/SGetSuggestedIDEWidget.cpp
@@ -18,7 +18,7 @@ TSharedPtr<SNotificationItem> SGetSuggestedIDEWidget::IDEDownloadNotification;
 
 void SGetSuggestedIDEWidget::Construct(const FArguments& InArgs)
 {
-	Visibility = TAttribute<EVisibility>(this, &SGetSuggestedIDEWidget::GetVisibility);
+	Visibility = InArgs._VisibilityOverride.IsSet() ? InArgs._VisibilityOverride : TAttribute<EVisibility>(this, &SGetSuggestedIDEWidget::GetVisibility);
 
 	ChildSlot
 	[
@@ -26,21 +26,30 @@ void SGetSuggestedIDEWidget::Construct(const FArguments& InArgs)
 	];
 }
 
-TSharedRef<SWidget> SGetSuggestedIDEWidget::CreateGetSuggestedIDEWidget() const
+TSharedRef<SWidget> SGetSuggestedIDEWidget::CreateGetSuggestedIDEWidget()
 {
 	if (FSourceCodeNavigation::GetCanDirectlyInstallSourceCodeIDE())
 	{
 		// If the installer for this platform's IDE can be downloaded and launched directly, show a button
-		return SNew(SButton)
+		return			
+			SNew(SButton)
+			.HAlign(HAlign_Center)
+			.VAlign(VAlign_Center)
+			.ButtonStyle(FAppStyle::Get(), "PrimaryButton")
+			.TextStyle(FAppStyle::Get(), "DialogButtonText")
 			.Text(FText::Format(LOCTEXT("IDEInstallButtonText", "Install {0}"), FSourceCodeNavigation::GetSuggestedSourceCodeIDE()))
-			.OnClicked(const_cast<SGetSuggestedIDEWidget*>(this), &SGetSuggestedIDEWidget::OnInstallIDEClicked);
+			.OnClicked(this, &SGetSuggestedIDEWidget::OnInstallIDEClicked);
 	}
 	else
 	{
-		// If the user must open a web page, show a link
-		return SNew(SHyperlink)
+		return	
+			SNew(SButton)
+			.HAlign(HAlign_Center)
+			.VAlign(VAlign_Center)
+			.ButtonStyle(FAppStyle::Get(), "PrimaryButton")
+			.TextStyle(FAppStyle::Get(), "DialogButtonText")
 			.Text(FText::Format(LOCTEXT("IDEDownloadLinkText", "Download {0}"), FSourceCodeNavigation::GetSuggestedSourceCodeIDE()))
-			.OnNavigate(const_cast<SGetSuggestedIDEWidget*>(this), &SGetSuggestedIDEWidget::OnDownloadIDEClicked, FSourceCodeNavigation::GetSuggestedSourceCodeIDEDownloadURL());
+			.OnClicked(this, &SGetSuggestedIDEWidget::OnDownloadIDEClicked, FSourceCodeNavigation::GetSuggestedSourceCodeIDEDownloadURL());
 	}
 }
 
@@ -49,9 +58,11 @@ EVisibility SGetSuggestedIDEWidget::GetVisibility() const
 	return FSourceCodeNavigation::IsCompilerAvailable() ? EVisibility::Collapsed : EVisibility::Visible;
 }
 
-void SGetSuggestedIDEWidget::OnDownloadIDEClicked(FString URL)
+FReply SGetSuggestedIDEWidget::OnDownloadIDEClicked(FString URL)
 {
 	FPlatformProcess::LaunchURL(*URL, NULL, NULL);
+
+	return FReply::Handled();
 }
 
 FReply SGetSuggestedIDEWidget::OnInstallIDEClicked()
diff --git a/Engine/Source/Editor/GameProjectGeneration/Private/SGetSuggestedIDEWidget.h b/Engine/Source/Editor/GameProjectGeneration/Private/SGetSuggestedIDEWidget.h
--- a/Engine/Source/Editor/GameProjectGeneration/Private/SGetSuggestedIDEWidget.h
+++ b/Engine/Source/Editor/GameProjectGeneration/Private/SGetSuggestedIDEWidget.h
@@ -19,6 +19,7 @@ class SGetSuggestedIDEWidget : public SCompoundWidget
 public:
 	SLATE_BEGIN_ARGS(SGetSuggestedIDEWidget)
 	{}
+		SLATE_ATTRIBUTE(EVisibility, VisibilityOverride)
 	SLATE_END_ARGS()
 
 	/** Constructs this widget with InArgs */
@@ -26,13 +27,13 @@ public:
 	
 private:
 	/** Creates the appropriate widget to display for the platform */
-	TSharedRef<SWidget> CreateGetSuggestedIDEWidget() const;
+	TSharedRef<SWidget> CreateGetSuggestedIDEWidget();
 
 	/** Gets the visibility of the global error label IDE Link */
 	EVisibility GetVisibility() const;
 
 	/** Handler for when the error label IDE hyperlink was clicked */
-	void OnDownloadIDEClicked(FString URL);
+	FReply OnDownloadIDEClicked(FString URL);
 
 	/** Handler for when the install IDE button was clicked */
 	FReply OnInstallIDEClicked();
diff --git a/Engine/Source/Editor/GameProjectGeneration/Private/SNewProjectWizard.cpp b/Engine/Source/Editor/GameProjectGeneration/Private/SNewProjectWizard.cpp
deleted file mode 100644
diff --git a/Engine/Source/Editor/GameProjectGeneration/Private/SNewProjectWizard.h b/Engine/Source/Editor/GameProjectGeneration/Private/SNewProjectWizard.h
deleted file mode 100644
diff --git a/Engine/Source/Editor/GameProjectGeneration/Private/SProjectBrowser.cpp b/Engine/Source/Editor/GameProjectGeneration/Private/SProjectBrowser.cpp
--- a/Engine/Source/Editor/GameProjectGeneration/Private/SProjectBrowser.cpp
+++ b/Engine/Source/Editor/GameProjectGeneration/Private/SProjectBrowser.cpp
@@ -50,9 +50,16 @@
 #include "ILauncherPlatform.h"
 #include "LauncherPlatformModule.h"
 #include "Interfaces/IMainFrameModule.h"
+#include "Internationalization/BreakIterator.h"
 
 #define LOCTEXT_NAMESPACE "ProjectBrowser"
 
+namespace ProjectBrowserDefs
+{
+	constexpr float ProjectTileHeight = 153;
+	constexpr float ProjectTileWidth = 102;
+	constexpr float ThumbnailSize = 64.0f, ThumbnailPadding = 5.f;
+}
 
 /**
  * Structure for project items.
@@ -63,25 +70,25 @@ struct FProjectItem
 	FText Description;
 	FText Category;
 	FString EngineIdentifier;
-	bool bUpToDate;
 	FString ProjectFile;
-	TSharedPtr<FSlateBrush> ProjectThumbnail;
-	bool bIsNewProjectItem;
 	TArray<FName> TargetPlatforms;
+	TSharedPtr<FSlateBrush> ProjectThumbnail;
+	FDateTime LastAccessTime;
+	bool bUpToDate;
 	bool bSupportsAllPlatforms;
-	bool bIsSample;
 
-	FProjectItem(const FText& InName, const FText& InDescription, const FString& InEngineIdentifier, bool InUpToDate, const TSharedPtr<FSlateBrush>& InProjectThumbnail, const FString& InProjectFile, bool InIsNewProjectItem, TArray<FName> InTargetPlatforms, bool InSupportsAllPlatforms)
+	FProjectItem(const FText& InName, const FText& InDescription, const FString& InEngineIdentifier, bool InUpToDate, const TSharedPtr<FSlateBrush>& InProjectThumbnail, const FString& InProjectFile, TArray<FName> InTargetPlatforms, bool InSupportsAllPlatforms)
 		: Name(InName)
 		, Description(InDescription)
+		, Category()
 		, EngineIdentifier(InEngineIdentifier)
-		, bUpToDate(InUpToDate)
 		, ProjectFile(InProjectFile)
-		, ProjectThumbnail(InProjectThumbnail)
-		, bIsNewProjectItem(InIsNewProjectItem)
 		, TargetPlatforms(InTargetPlatforms)
+		, ProjectThumbnail(InProjectThumbnail)
+		, bUpToDate(InUpToDate)
+		//, bIsNewProjectItem(InIsNewProjectItem)
 		, bSupportsAllPlatforms(InSupportsAllPlatforms)
-	{ }
+	{}
 
 	/** Check if this project is up to date */
 	bool IsUpToDate() const
@@ -108,17 +115,118 @@ struct FProjectItem
 };
 
 
-/**
- * Structure for project categories.
- */
-struct FProjectCategory
+class SProjectTile : public STableRow<TSharedPtr<FProjectItem>>
 {
-	FText CategoryName;
-	TSharedPtr< STileView< TSharedPtr<FProjectItem> > > ProjectTileView;
-	TArray< TSharedPtr<FProjectItem> > ProjectItemsSource;
-	TArray< TSharedPtr<FProjectItem> > FilteredProjectItemsSource;
-};
+	SLATE_BEGIN_ARGS(SProjectTile)
+	{}
+	SLATE_END_ARGS()
 
+public:
+	void Construct(const FArguments& InArgs, TSharedPtr<FProjectItem> ProjectItem, const TSharedRef<STableViewBase>& OwnerTable)
+	{
+		STableRow::Construct(
+			STableRow::FArguments()
+			.Style(FAppStyle::Get(), "ProjectBrowser.TableRow")
+			.Padding(2.0f)
+			.Content()
+		[
+			SNew(SBorder)
+			.Padding(FMargin(0.0f, 0.0f, 5.0f, 5.0f))
+			.BorderImage(FAppStyle::Get().GetBrush("ProjectBrowser.ProjectTile.DropShadow"))
+			[
+				SNew(SOverlay)
+				+ SOverlay::Slot()
+				[
+					SNew(SVerticalBox)
+					// Thumbnail
+					+ SVerticalBox::Slot()
+					.AutoHeight()
+					.HAlign(HAlign_Center)
+					.VAlign(VAlign_Center)
+					[
+						SNew(SBox)
+						.WidthOverride(102)
+						.HeightOverride(102)
+						[
+							SNew(SBorder)
+							.Padding(0)
+							.BorderImage(FAppStyle::Get().GetBrush("ProjectBrowser.ProjectTile.ThumbnailAreaBackground"))
+							.HAlign(HAlign_Center)
+							.VAlign(VAlign_Center)
+							[
+								SNew(SImage)
+								.Image(ProjectItem->ProjectThumbnail ? ProjectItem->ProjectThumbnail.Get() : FAppStyle::Get().GetBrush("GameProjectDialog.DefaultGameThumbnail.Small"))
+							]
+						]
+					]
+					// Name
+					+ SVerticalBox::Slot()
+					[
+						SNew(SBorder)
+						.Padding(FMargin(ProjectBrowserDefs::ThumbnailPadding, 0))
+						.Padding(FMargin(3.0f, 3.0f))
+						.BorderImage(FAppStyle::Get().GetBrush("ProjectBrowser.ProjectTile.NameAreaBackground"))
+						[
+							SNew(SVerticalBox)
+							+SVerticalBox::Slot()
+							[
+								SNew(STextBlock)
+								.Font(FAppStyle::Get().GetFontStyle("ProjectBrowser.ProjectTile.Font"))
+								.WrapTextAt(ProjectBrowserDefs::ProjectTileWidth - 4.0f)
+								.LineBreakPolicy(FBreakIterator::CreateCamelCaseBreakIterator())
+								.Text(ProjectItem->Name)
+								.ColorAndOpacity(FAppStyle::Get().GetSlateColor("Colors.Foreground"))
+							]
+							+SVerticalBox::Slot()
+							.AutoHeight()
+							.Padding(0.0f, 4.0f, 0.0f, 4.0f)
+							.VAlign(VAlign_Bottom)
+							[
+								SNew(STextBlock)
+								.Text(FText::FromString(ProjectItem->GetEngineLabel()))
+								.Font(FAppStyle::Get().GetFontStyle("ProjectBrowser.ProjectTile.Font"))
+								.ColorAndOpacity(FAppStyle::Get().GetSlateColor("Colors.White25"))
+								.Visibility(ProjectItem->IsUpToDate() ? EVisibility::Collapsed : EVisibility::Visible)
+							]
+						]
+					]
+				]
+				+ SOverlay::Slot()
+				[
+					SNew(SImage)
+					.Visibility(EVisibility::HitTestInvisible)
+					.Image_Lambda
+					(
+						[this]()
+						{
+							const bool bIsSelected = IsSelected();
+							const bool bIsRowHovered = IsHovered();
+
+							if (bIsSelected && bIsRowHovered)
+							{
+								static const FName SelectedHover("ProjectBrowser.ProjectTile.SelectedHoverBorder");
+								return FAppStyle::Get().GetBrush(SelectedHover);
+							}
+							else if (bIsSelected)
+							{
+								static const FName Selected("ProjectBrowser.ProjectTile.SelectedBorder");
+								return FAppStyle::Get().GetBrush(Selected);
+							}
+							else if (bIsRowHovered)
+							{
+								static const FName Hovered("ProjectBrowser.ProjectTile.HoverBorder");
+								return FAppStyle::Get().GetBrush(Hovered);
+							}
+
+							return FStyleDefaults::GetNoBrush();
+						}
+					)
+				]
+			]
+		],
+		OwnerTable);
+	}
+};
 
 void ProjectItemToString(const TSharedPtr<FProjectItem> InItem, TArray<FString>& OutFilterStrings)
 {
@@ -131,25 +239,32 @@ void ProjectItemToString(const TSharedPtr<FProjectItem> InItem, TArray<FString>&
 
 SProjectBrowser::SProjectBrowser()
 	: ProjectItemFilter( ProjectItemTextFilter::FItemToStringArray::CreateStatic( ProjectItemToString ))
-	, NumFilteredProjects(0)
-	, ThumbnailBorderPadding(8)
-	, ThumbnailSize(128)
 	, bPreventSelectionChangeEvent(false)
 {
 
 }
 
-
-BEGIN_SLATE_FUNCTION_BUILD_OPTIMIZATION
 void SProjectBrowser::Construct( const FArguments& InArgs )
 {
-	// Prepare the categories box
-	CategoriesBox = SNew(SVerticalBox);
+	// Prepare the projects box
+	ProjectsBox = SNew(SVerticalBox);
+
+	SAssignNew(ProjectTileView, STileView<TSharedPtr<FProjectItem>>)
+		.ListItemsSource(&FilteredProjectItemsSource)
+		.SelectionMode(ESelectionMode::Single)
+		.ClearSelectionOnClick(false)
+		.ItemAlignment(EListItemAlignment::LeftAligned)
+		.OnGenerateTile(this, &SProjectBrowser::MakeProjectViewWidget)
+		.OnContextMenuOpening(this, &SProjectBrowser::OnGetContextMenuContent)
+		.OnMouseButtonDoubleClick(this, &SProjectBrowser::HandleProjectItemDoubleClick)
+		.OnSelectionChanged(this, &SProjectBrowser::HandleProjectViewSelectionChanged)
+		.ItemHeight(ProjectBrowserDefs::ProjectTileHeight+9)
+		.ItemWidth(ProjectBrowserDefs::ProjectTileWidth+9);
 
 	// Find all projects
 	FindProjects();
 
-	CategoriesBox->AddSlot()
+	ProjectsBox->AddSlot()
 	.HAlign(HAlign_Center)
 	.Padding(FMargin(0.f, 25.f))
 	[
@@ -158,7 +273,7 @@ void SProjectBrowser::Construct( const FArguments& InArgs )
 		.Text(LOCTEXT("NoProjects", "You don't have any projects yet :("))
 	];
 
-	CategoriesBox->AddSlot()
+	ProjectsBox->AddSlot()
 	.HAlign(HAlign_Center)
 	.Padding(FMargin(0.f, 25.f))
 	[
@@ -167,250 +282,57 @@ void SProjectBrowser::Construct( const FArguments& InArgs )
 		.Text(LOCTEXT("NoProjectsAfterFilter", "There are no projects that match the specified filter"))
 	];
 
+	ProjectsBox->AddSlot()
+	[
+		ProjectTileView.ToSharedRef()
+	];
+
 	ChildSlot
 	[
-		SNew(SBorder)
-		.BorderImage( FEditorStyle::GetBrush("ToolPanel.GroupBorder") )
-		.Padding(FMargin(8.f, 4.f))
+		SNew(SVerticalBox)
+		+ SVerticalBox::Slot()
+		.AutoHeight()
 		[
-			SNew(SVerticalBox)
-
-			// Categories
-			+ SVerticalBox::Slot()
-			.Padding(8.f)
-			.FillHeight(1.0f)
+			SNew(SHorizontalBox)	
+			+ SHorizontalBox::Slot()
+			.Padding(FMargin(0, 0, 5.f, 5))
+			.VAlign(VAlign_Center)
 			[
-				SNew(SVerticalBox)
-
-				+ SVerticalBox::Slot()
-				.AutoHeight()
-				.VAlign(VAlign_Top)
-				[
-					SNew(SHorizontalBox)
-					
-					+ SHorizontalBox::Slot()
-					.Padding(FMargin(0, 0, 5.f, 0))
-					.VAlign(VAlign_Center)
-					[
-						SNew(SOverlay)
-
-						+SOverlay::Slot()
-						[
-							SAssignNew(SearchBoxPtr, SSearchBox)
-							.HintText(LOCTEXT("FilterHint", "Filter Projects..."))
-							.OnTextChanged(this, &SProjectBrowser::OnFilterTextChanged)
-						]
-
-						+SOverlay::Slot()
-						[
-							SNew(SBorder)
-							.Visibility(this, &SProjectBrowser::GetFilterActiveOverlayVisibility)
-							.BorderImage(FEditorStyle::Get().GetBrush("SearchBox.ActiveBorder"))
-						]
-					]
-
-					+SHorizontalBox::Slot()
-					.AutoWidth()
-					.VAlign(VAlign_Center)
-					.Padding(FMargin(0, 0, 5.f, 0))
-					[
-						SNew(SButton)
-						.ButtonStyle(FEditorStyle::Get(), "ToggleButton")
-						.OnClicked(this, &SProjectBrowser::FindProjects)
-						.ForegroundColor(FSlateColor::UseForeground())
-						.ToolTipText(LOCTEXT("RefreshProjectList", "Refresh the project list"))
-						.HAlign(HAlign_Center)
-						.VAlign(VAlign_Center)
-						[
-							SNew(SHorizontalBox)
-
-							+ SHorizontalBox::Slot()
-							.Padding(2.0f)
-							.VAlign(VAlign_Center)
-							.AutoWidth()
-							[
-								SNew(SImage)
-								.Image(FEditorStyle::GetBrush("Icons.Refresh")) 
-							]
-
-							+ SHorizontalBox::Slot()
-							.VAlign(VAlign_Center)
-							.Padding(2.0f)
-							[
-								SNew(STextBlock)
-								.TextStyle(FEditorStyle::Get(), "ProjectBrowser.Toolbar.Text")
-								.Text(LOCTEXT("RefreshProjectsText", "Refresh"))
-							]
-						]
-					]
-
-					+SHorizontalBox::Slot()
-					.AutoWidth()
-					[
-						SNew(SButton)
-						.Visibility( FLauncherPlatformModule::Get()->CanOpenLauncher(true) ? EVisibility::Visible : EVisibility::Collapsed )
-						.ButtonStyle(FEditorStyle::Get(), "ToggleButton")
-						.OnClicked(this, &SProjectBrowser::HandleMarketplaceTabButtonClicked)
-						.ForegroundColor(FSlateColor::UseForeground())
-						.ToolTipText(LOCTEXT("MarketplaceToolTip", "Check out the Marketplace to find new projects!"))
-						.HAlign(HAlign_Center)
-						.VAlign(VAlign_Center)
-						[
-							SNew(SHorizontalBox)
-
-							+ SHorizontalBox::Slot()
-							.Padding(2.0f)
-							.VAlign(VAlign_Center)
-							.AutoWidth()
-							[
-								SNew(SImage)
-								.Image(FEditorStyle::GetBrush("LevelEditor.OpenMarketplace.Small")) 
-							]
-
-							+ SHorizontalBox::Slot()
-							.VAlign(VAlign_Center)
-							.Padding(2.0f)
-							[
-								SNew(STextBlock)
-								.TextStyle(FEditorStyle::Get(), "ProjectBrowser.Toolbar.Text")
-								.Text(LOCTEXT("Marketplace", "Marketplace"))
-							]
-						]
-					]
-				]
-
-				+ SVerticalBox::Slot()
-				.Padding(FMargin(0, 5.f))
-				[
-					SNew(SScrollBox)
-
-					+ SScrollBox::Slot()
-					[
-						CategoriesBox.ToSharedRef()
-					]
-				]
+				SAssignNew(SearchBoxPtr, SSearchBox)
+				.HintText(LOCTEXT("FilterHint", "Filter Projects..."))
+				.OnTextChanged(this, &SProjectBrowser::OnFilterTextChanged)
 			]
-
-			+ SVerticalBox::Slot()
-			.Padding( 0, 20, 0, 0 )	// Lots of vertical padding before the dialog buttons at the bottom
-			.AutoHeight()
+			+SHorizontalBox::Slot()
+			.AutoWidth()
+			.VAlign(VAlign_Center)
+			.Padding(FMargin(0, 0, 5.f, 0))
 			[
-				SNew( SHorizontalBox )
-
-				// Auto-load project
-				+ SHorizontalBox::Slot()
-				.AutoWidth()
+				SNew(SButton)
+				.ButtonStyle(FAppStyle::Get(), "SimpleButton")
+				.OnClicked(this, &SProjectBrowser::FindProjects)
+				.ToolTipText(LOCTEXT("RefreshProjectList", "Refresh the project list"))
+				.HAlign(HAlign_Center)
 				.VAlign(VAlign_Center)
 				[
-					SNew(SCheckBox)			
-					.IsChecked(GetDefault<UEditorSettings>()->bLoadTheMostRecentlyLoadedProjectAtStartup ? ECheckBoxState::Checked : ECheckBoxState::Unchecked)
-					.OnCheckStateChanged(this, &SProjectBrowser::OnAutoloadLastProjectChanged)
-					.Content()
-					[
-						SNew(STextBlock).Text(LOCTEXT("AutoloadOnStartupCheckbox", "Always load last project on startup"))
-					]
-				]
-
-				+SHorizontalBox::Slot()
-				.FillWidth(1.0f)
-				[
-					SNullWidget::NullWidget
-				]
-				// Browse Button
-				+SHorizontalBox::Slot()
-				.AutoWidth()
-				.Padding(8.f, 0)
-				[
-					SNew(SButton)
-					.Text(LOCTEXT("BrowseProjectButton", "Browse..."))
-					.OnClicked(this, &SProjectBrowser::OnBrowseToProjectClicked)
-					.ContentPadding(FCoreStyle::Get().GetMargin("StandardDialog.ContentPadding"))
-				]
-
-				// Open Button
-				+SHorizontalBox::Slot()
-				.AutoWidth()
-				[
-					SNew(SButton)
-					.Text(LOCTEXT("OpenProjectButton", "Open"))
-					.OnClicked(this, &SProjectBrowser::HandleOpenProjectButtonClicked)
-					.IsEnabled(this, &SProjectBrowser::HandleOpenProjectButtonIsEnabled)
-					.ContentPadding(FCoreStyle::Get().GetMargin("StandardDialog.ContentPadding"))
-					.Visibility(InArgs._HideOpenButton ? EVisibility::Collapsed : EVisibility::Visible)
+					SNew(SImage)
+					.Image(FEditorStyle::GetBrush("Icons.Refresh"))
+					.ColorAndOpacity(FSlateColor::UseForeground())
 				]
 			]
 		]
+		+ SVerticalBox::Slot()
+		.Padding(FMargin(0, 5.f))
+		[
+			ProjectsBox.ToSharedRef()
+		]
 	];
 
-	// Select the first item in the first category
-	if (ProjectCategories.Num() > 0)
-	{
-		const TSharedRef<FProjectCategory> Category = ProjectCategories[0];
-		if ( ensure(Category->ProjectItemsSource.Num() > 0) && ensure(Category->ProjectTileView.IsValid()))
-		{
-			Category->ProjectTileView->SetSelection(Category->ProjectItemsSource[0], ESelectInfo::Direct);
-		}
-	}
-
-	bHasProjectFiles = false;
-	for (auto CategoryIt = ProjectCategories.CreateConstIterator(); CategoryIt; ++CategoryIt)
-	{
-		const TSharedRef<FProjectCategory>& Category = *CategoryIt;
-
-		if (Category->ProjectItemsSource.Num() > 0)
-		{
-			bHasProjectFiles = true;
-			break;
-		}
-	}
-
 	ProjectSelectionChangedDelegate = InArgs._OnSelectionChanged;
 }
-END_SLATE_FUNCTION_BUILD_OPTIMIZATION
-
-void SProjectBrowser::ConstructCategory( const TSharedRef<SVerticalBox>& InCategoriesBox, const TSharedRef<FProjectCategory>& Category )
-{
-	// Title
-	InCategoriesBox->AddSlot()
-	.AutoHeight()
-	[
-		SNew(STextBlock)
-		.Visibility(this, &SProjectBrowser::GetProjectCategoryVisibility, Category)
-		.TextStyle(FEditorStyle::Get(), "GameProjectDialog.ProjectNamePathLabels")
-		.Text(Category->CategoryName)
-	];
-
-	// Separator
-	InCategoriesBox->AddSlot()
-	.AutoHeight()
-	.Padding(0, 2.0f, 0, 8.0f)
-	[
-		SNew(SSeparator)
-		.Visibility(this, &SProjectBrowser::GetProjectCategoryVisibility, Category)
-	];
-
-	// Project tile view
-	InCategoriesBox->AddSlot()
-	.AutoHeight()
-	[
-		SAssignNew(Category->ProjectTileView, STileView<TSharedPtr<FProjectItem> >)
-		.Visibility(this, &SProjectBrowser::GetProjectCategoryVisibility, Category)
-		.ListItemsSource(&Category->FilteredProjectItemsSource)
-		.SelectionMode(ESelectionMode::Single)
-		.ClearSelectionOnClick(true)
-		.AllowOverscroll(EAllowOverscroll::No)
-		.OnGenerateTile(this, &SProjectBrowser::MakeProjectViewWidget)
-		.OnContextMenuOpening(this, &SProjectBrowser::OnGetContextMenuContent)
-		.OnMouseButtonDoubleClick(this, &SProjectBrowser::HandleProjectItemDoubleClick)
-		.OnSelectionChanged(this, &SProjectBrowser::HandleProjectViewSelectionChanged, Category->CategoryName)
-		.ItemHeight(ThumbnailSize + ThumbnailBorderPadding + 32)
-		.ItemWidth(ThumbnailSize + ThumbnailBorderPadding)
-	];
-}
 
 bool SProjectBrowser::HasProjects() const
 {
-	return bHasProjectFiles;
+	return ProjectItemsSource.Num() > 0;
 }
 
 FString SProjectBrowser::GetSelectedProjectFile() const
@@ -429,100 +351,8 @@ FString SProjectBrowser::GetSelectedProjectFile() const
 
 TSharedRef<ITableRow> SProjectBrowser::MakeProjectViewWidget(TSharedPtr<FProjectItem> ProjectItem, const TSharedRef<STableViewBase>& OwnerTable)
 {
-	if ( !ensure(ProjectItem.IsValid()) )
-	{
-		return SNew( STableRow<TSharedPtr<FProjectItem>>, OwnerTable );
-	}
-
-	TSharedPtr<SWidget> Thumbnail;
-	if ( ProjectItem->bIsNewProjectItem )
-	{
-		Thumbnail = SNew(SBox)
-			.Padding(ThumbnailBorderPadding)
-			[
-				SNew(SBorder)
-				.Padding(0)
-				.BorderImage(FEditorStyle::GetBrush("MarqueeSelection"))
-				[
-					SNew(SBox)
-					.HAlign(HAlign_Center)
-					.VAlign(VAlign_Center)
-					[
-						SNew(STextBlock)
-						.TextStyle( FEditorStyle::Get(), "GameProjectDialog.NewProjectTitle" )
-						.Text( LOCTEXT("NewProjectThumbnailText", "NEW") )
-					]
-				]
-			];
-	}
-	else
-	{
-		const FLinearColor Tint = ProjectItem->IsUpToDate() ? FLinearColor::White : FLinearColor::White.CopyWithNewOpacity(0.5);
-
-		// Drop shadow border
-		Thumbnail =	SNew(SBorder)
-			.Padding(ThumbnailBorderPadding)
-			.BorderImage( FEditorStyle::GetBrush("ContentBrowser.ThumbnailShadow") )
-			.ColorAndOpacity(Tint)
-			.BorderBackgroundColor(Tint)
-			[
-				SNew(SImage).Image(this, &SProjectBrowser::GetProjectItemImage, TWeakPtr<FProjectItem>(ProjectItem))
-			];
-	}
-
-	TSharedRef<ITableRow> TableRow = SNew( STableRow<TSharedPtr<FProjectItem>>, OwnerTable )
-	.Style(FEditorStyle::Get(), "GameProjectDialog.TemplateListView.TableRow")
-	[
-		SNew(SBox)
-		.HeightOverride(ThumbnailSize+ThumbnailBorderPadding+5)
-		[
-			SNew(SVerticalBox)
-
-			// Thumbnail
-			+SVerticalBox::Slot()
-			.AutoHeight()
-			[
-				SNew(SBox)
-				.WidthOverride(ThumbnailSize + ThumbnailBorderPadding * 2)
-				.HeightOverride(ThumbnailSize + ThumbnailBorderPadding * 2)
-				[
-					SNew(SOverlay)
-
-					+ SOverlay::Slot()
-					[
-						Thumbnail.ToSharedRef()
-					]
-
-					// Show the out of date engine version for this project file
-					+ SOverlay::Slot()
-					.HAlign(HAlign_Right)
-					.VAlign(VAlign_Bottom)
-					.Padding(10)
-					[
-						SNew(STextBlock)
-						.Text(FText::FromString(ProjectItem->GetEngineLabel()))
-						.TextStyle(FEditorStyle::Get(), "ProjectBrowser.VersionOverlayText")
-						.ColorAndOpacity(FLinearColor::White.CopyWithNewOpacity(0.5f))
-						.Visibility(ProjectItem->IsUpToDate() ? EVisibility::Collapsed : EVisibility::Visible)
-					]
-				]
-			]
-
-			// Name
-			+SVerticalBox::Slot()
-			.HAlign(HAlign_Center)
-			.VAlign(VAlign_Top)
-			[
-				SNew(STextBlock)
-				.HighlightText(this, &SProjectBrowser::GetItemHighlightText)
-				.Text(ProjectItem->Name)
-			]
-		]
-	];
-
-	TableRow->AsWidget()->SetToolTip(MakeProjectToolTip(ProjectItem));
-
-	return TableRow;
+	return SNew(SProjectTile, ProjectItem, OwnerTable)
+			.ToolTip(MakeProjectToolTip(ProjectItem));
 }
 
 
@@ -730,17 +560,13 @@ const FSlateBrush* SProjectBrowser::GetProjectItemImage(TWeakPtr<FProjectItem> P
 
 TSharedPtr<FProjectItem> SProjectBrowser::GetSelectedProjectItem() const
 {
-	for ( auto CategoryIt = ProjectCategories.CreateConstIterator(); CategoryIt; ++CategoryIt )
+	TArray< TSharedPtr<FProjectItem> > SelectedItems = ProjectTileView->GetSelectedItems();
+	if (SelectedItems.Num() > 0)
 	{
-		const TSharedRef<FProjectCategory>& Category = *CategoryIt;
-		TArray< TSharedPtr<FProjectItem> > SelectedItems = Category->ProjectTileView->GetSelectedItems();
-		if ( SelectedItems.Num() > 0 )
-		{
-			return SelectedItems[0];
-		}
+		return SelectedItems[0];	
 	}
 	
-	return NULL;
+	return nullptr;
 }
 
 
@@ -773,7 +599,11 @@ static TSharedPtr<FSlateDynamicImageBrush> GetThumbnailForProject(const FString&
 	if (!PNGFileToUse.IsEmpty())
 	{
 		const FName BrushName = FName(*PNGFileToUse);
-		DynamicBrush = MakeShareable(new FSlateDynamicImageBrush(BrushName, FVector2D(128, 128)));
+		DynamicBrush = MakeShared<FSlateDynamicImageBrush>(BrushName, FVector2D(128, 128));
+		DynamicBrush->OutlineSettings.CornerRadii = FVector4(4, 4, 0, 0);
+		DynamicBrush->OutlineSettings.RoundingType = ESlateBrushRoundingType::FixedRadius;
+		DynamicBrush->DrawAs = ESlateBrushDrawType::RoundedBox;
+
 	}
 
 	return DynamicBrush;
@@ -809,10 +639,9 @@ static TSharedPtr<FProjectItem> CreateProjectItem(const FString& ProjectFilename
 			}
 			TargetPlatforms.Sort(FNameLexicalLess());
 
-			const bool bIsNewProjectItem = false;
 			TSharedPtr<FProjectItem> ProjectItem = MakeShareable(new FProjectItem(FText::FromString(ProjectStatus.Name),
 				FText::FromString(ProjectStatus.Description),
-				ProjectEngineIdentifier, bIsUpToDate, DynamicBrush, ProjectFilename, bIsNewProjectItem, 
+				ProjectEngineIdentifier, bIsUpToDate, DynamicBrush, ProjectFilename, 
 				TargetPlatforms, ProjectStatus.SupportsAllPlatforms()));
 
 			const FText SamplesCategoryName = LOCTEXT("SamplesCategoryName", "Samples");
@@ -822,6 +651,8 @@ static TSharedPtr<FProjectItem> CreateProjectItem(const FString& ProjectFilename
 				ProjectItem->Category = SamplesCategoryName;
 			}
 
+			ProjectItem->LastAccessTime = IFileManager::Get().GetAccessTimeStamp(*ProjectFilename);
+
 			return ProjectItem;
 		}
 	}
@@ -831,15 +662,7 @@ static TSharedPtr<FProjectItem> CreateProjectItem(const FString& ProjectFilename
 
 FReply SProjectBrowser::FindProjects()
 {
-	enum class EProjectCategoryType : uint8
-	{
-		Sample,
-		UserDefined,
-	};
-
-	ProjectCategories.Empty();
-	CategoriesBox->ClearChildren();
-
+#if 0
 	// Create a map of parent project folders to their category
 	TMap<FString, EProjectCategoryType> ProjectFilesToCategoryType;
 
@@ -970,58 +793,116 @@ FReply SProjectBrowser::FindProjects()
 		Category->ProjectItemsSource.Sort( FCompareProjectItems() );
 	}
 
-	PopulateFilteredProjectCategories();
+	PopulateFilteredProjects();
 
 	for (auto CategoryIt = ProjectCategories.CreateConstIterator(); CategoryIt; ++CategoryIt)
 	{
 		const TSharedRef<FProjectCategory>& Category = *CategoryIt;
 		ConstructCategory(CategoriesBox.ToSharedRef(), Category);
 	}
+#else
 
-	return FReply::Handled();
-}
+	const FString LastSelectedProjectFile = CurrentSelectedProjectPath;
 
+	ProjectItemsSource.Empty();
+	FilteredProjectItemsSource.Empty();
+	ProjectTileView->RequestListRefresh();
 
-void SProjectBrowser::AddProjectToCategory(const TSharedRef<FProjectItem>& ProjectItem, const FText& ProjectCategory)
-{
-	for ( auto CategoryIt = ProjectCategories.CreateConstIterator(); CategoryIt; ++CategoryIt )
+	TArray<FRecentProjectFile> RecentProjects = GetDefault<UEditorSettings>()->RecentlyOpenedProjectFiles;
+
+	//TSet<FString> ExistingRecentProjects = RecentProjects;
+
+	TSet<FString> AllFoundProjectFiles;
+
+	//AllFoundProjectFiles.Append(RecentProjects);
+
+	// Find all the engine installations
+	TMap<FString, FString> EngineInstallations;
+	FDesktopPlatformModule::Get()->EnumerateEngineInstallations(EngineInstallations);
+
+	// Add projects from every branch that we know about
+	FString CurrentEngineIdentifier = FDesktopPlatformModule::Get()->GetCurrentEngineIdentifier();
+	for (TMap<FString, FString>::TConstIterator Iter(EngineInstallations); Iter; ++Iter)
 	{
-		const TSharedRef<FProjectCategory>& Category = *CategoryIt;
-		if ( Category->CategoryName.EqualToCaseIgnored(ProjectCategory) )
+		TArray<FString> ProjectFiles;
+		if (FDesktopPlatformModule::Get()->EnumerateProjectsKnownByEngine(Iter.Key(), false, ProjectFiles))
 		{
-			Category->ProjectItemsSource.Add(ProjectItem);
-			return;
+			AllFoundProjectFiles.Append(MoveTemp(ProjectFiles));
 		}
 	}
 
-	TSharedRef<FProjectCategory> NewCategory = MakeShareable( new FProjectCategory );
-	NewCategory->CategoryName = ProjectCategory;
-	NewCategory->ProjectItemsSource.Add(ProjectItem);
-	ProjectCategories.Add(NewCategory);
-}
-
+	// Add all the samples from the launcher
+	TArray<FString> LauncherSampleProjects;
+	FDesktopPlatformModule::Get()->EnumerateLauncherSampleProjects(LauncherSampleProjects);
+	AllFoundProjectFiles.Append(MoveTemp(LauncherSampleProjects));
 
-void SProjectBrowser::PopulateFilteredProjectCategories()
-{
-	NumFilteredProjects = 0;
-	for (auto& Category : ProjectCategories)
+	// Add all the native project files we can find
+	FUProjectDictionary& DefaultProjectDictionary = FUProjectDictionary::GetDefault();
+	DefaultProjectDictionary.Refresh();
+	const TArray<FString>& NativeProjectFiles = DefaultProjectDictionary.GetProjectPaths();
+	for(const FString& ProjectFile : NativeProjectFiles)
 	{
-		Category->FilteredProjectItemsSource.Empty();
+		if (!ProjectFile.Contains(TEXT("/Templates/")))
+		{
+			AllFoundProjectFiles.Add(ProjectFile);
+		}
+	}
 
-		for (auto& ProjectItem : Category->ProjectItemsSource)
+	TSharedPtr<FProjectItem> NewProjectToSelect;
+
+	// Normalize all the filenames and make sure there are no duplicates
+	for (FString& ProjectFile : AllFoundProjectFiles)
+	{
+		FString ProjectFilename = IFileManager::Get().ConvertToAbsolutePathForExternalAppForRead(*ProjectFile);
+		TSharedPtr<FProjectItem> NewProjectItem = CreateProjectItem(ProjectFilename);
+		if (NewProjectItem)
 		{
-			if (ProjectItemFilter.PassesFilter(ProjectItem))
+			if (NewProjectItem->ProjectFile == LastSelectedProjectFile)
+			{
+				NewProjectToSelect = NewProjectItem;
+			}
+
+			// Get a valid last access time. The editor will set this for recent projects and is more accurate than the NTFS access time
+			if (FRecentProjectFile* RecentProject = RecentProjects.FindByKey(NewProjectItem->ProjectFile))
 			{
-				Category->FilteredProjectItemsSource.Add(ProjectItem);
-				++NumFilteredProjects;
+				NewProjectItem->LastAccessTime = RecentProject->LastOpenTime;
 			}
+
+			ProjectItemsSource.Add(NewProjectItem);
 		}
+	}
+
+	ProjectItemsSource.Sort([](auto& A, auto& B) { return A->LastAccessTime > B->LastAccessTime; });
 
-		if (Category->ProjectTileView.IsValid())
+	PopulateFilteredProjects();
+
+	if (NewProjectToSelect && FilteredProjectItemsSource.Contains(NewProjectToSelect))
+	{
+		ProjectTileView->SetSelection(NewProjectToSelect, ESelectInfo::Direct);
+	}
+	else if (FilteredProjectItemsSource.Num() > 0)
+	{
+		ProjectTileView->SetSelection(FilteredProjectItemsSource[0], ESelectInfo::Direct);
+	}
+
+#endif
+
+	return FReply::Handled();
+}
+
+void SProjectBrowser::PopulateFilteredProjects()
+{
+	FilteredProjectItemsSource.Empty();
+
+	for (TSharedPtr<FProjectItem>& ProjectItem : ProjectItemsSource)
+	{
+		if (ProjectItemFilter.PassesFilter(ProjectItem))
 		{
-			Category->ProjectTileView->RequestListRefresh();
+			FilteredProjectItemsSource.Add(ProjectItem);
 		}
 	}
+
+	ProjectTileView->RequestListRefresh();
 }
 
 FReply SProjectBrowser::OnKeyDown( const FGeometry& MyGeometry, const FKeyEvent& InKeyEvent )
@@ -1260,18 +1141,18 @@ bool SProjectBrowser::OpenProject( const FString& InProjectFile )
 
 void SProjectBrowser::OpenSelectedProject( )
 {
-	if ( CurrentSelectedProjectPath.IsEmpty() )
+	if (CurrentSelectedProjectPath.IsEmpty())
 	{
 		return;
 	}
 
-	OpenProject( CurrentSelectedProjectPath.ToString() );
+	OpenProject(CurrentSelectedProjectPath);
 }
 
 /* SProjectBrowser event handlers
  *****************************************************************************/
 
-FReply SProjectBrowser::HandleOpenProjectButtonClicked( )
+FReply SProjectBrowser::OnOpenProject( )
 {
 	OpenSelectedProject();
 
@@ -1290,7 +1171,7 @@ void SProjectBrowser::HandleProjectItemDoubleClick( TSharedPtr<FProjectItem> Tem
 	OpenSelectedProject();
 }
 
-FReply SProjectBrowser::OnBrowseToProjectClicked()
+FReply SProjectBrowser::OnBrowseToProject()
 {
 	const FString ProjectFileDescription = LOCTEXT( "FileTypeDescription", "Unreal Project File" ).ToString();
 	const FString ProjectFileExtension = FString::Printf(TEXT("*.%s"), *FProjectDescriptor::GetExtension());
@@ -1298,12 +1179,12 @@ FReply SProjectBrowser::OnBrowseToProjectClicked()
 
 	// Find the first valid project file to select by default
 	FString DefaultFolder = FEditorDirectories::Get().GetLastDirectory(ELastDirectory::PROJECT);
-	for ( auto ProjectIt = GetDefault<UEditorSettings>()->RecentlyOpenedProjectFiles.CreateConstIterator(); ProjectIt; ++ProjectIt )
+	for (const FRecentProjectFile& RecentFile : GetDefault<UEditorSettings>()->RecentlyOpenedProjectFiles)
 	{
-		if ( IFileManager::Get().FileSize(**ProjectIt) > 0 )
+		if ( IFileManager::Get().FileSize(*RecentFile.ProjectName) > 0 )
 		{
 			// This is the first uproject file in the recents list that actually exists
-			DefaultFolder = FPaths::GetPath(*ProjectIt);
+			DefaultFolder = FPaths::GetPath(*RecentFile.ProjectName);
 			break;
 		}
 	}
@@ -1334,17 +1215,17 @@ FReply SProjectBrowser::OnBrowseToProjectClicked()
 		);
 	}
 
-	if ( bOpened && OpenFilenames.Num() > 0 )
+	if (bOpened && OpenFilenames.Num() > 0)
 	{
-		HandleProjectViewSelectionChanged( NULL, ESelectInfo::Direct, FText() );
+		HandleProjectViewSelectionChanged(nullptr, ESelectInfo::Direct);
 
 		FString Path = OpenFilenames[0];
 		if ( FPaths::IsRelative( Path ) )
 		{
-			Path = FPaths::ConvertRelativePathToFull( Path );
+			Path = FPaths::ConvertRelativePathToFull(Path);
 		}
 
-		CurrentSelectedProjectPath = FText::FromString( Path );
+		CurrentSelectedProjectPath = Path;
 
 		OpenSelectedProject();
 	}
@@ -1352,51 +1233,24 @@ FReply SProjectBrowser::OnBrowseToProjectClicked()
 	return FReply::Handled();
 }
 
-void SProjectBrowser::ClearSelection()
+void SProjectBrowser::HandleProjectViewSelectionChanged(TSharedPtr<FProjectItem> ProjectItem, ESelectInfo::Type SelectInfo)
 {
-	TGuardValue<bool> SelectionEventGuard(bPreventSelectionChangeEvent, true);
-
-	for (const TSharedRef<FProjectCategory>& Category : ProjectCategories)
-	{
-		if (Category->ProjectTileView.IsValid())
-		{
-			Category->ProjectTileView->ClearSelection();
-		}
-	}
-}
-
-
-void SProjectBrowser::HandleProjectViewSelectionChanged(TSharedPtr<FProjectItem> ProjectItem, ESelectInfo::Type SelectInfo, FText CategoryName)
-{
-	if ( !bPreventSelectionChangeEvent )
-	{
-		TGuardValue<bool> SelectionEventGuard(bPreventSelectionChangeEvent, true);
-
-		for (const TSharedRef<FProjectCategory>& Category : ProjectCategories)
-		{
-			if (Category->ProjectTileView.IsValid() && !Category->CategoryName.EqualToCaseIgnored(CategoryName))
-			{
-				Category->ProjectTileView->ClearSelection();
-			}
-		}
-	}
-
 	FString ProjectFile;
 	if (ProjectItem.IsValid())
 	{
 		ProjectFile = ProjectItem->ProjectFile;
-		CurrentSelectedProjectPath = FText::FromString(ProjectFile);
+		CurrentSelectedProjectPath = ProjectFile;
 	}
 	else
 	{
-		CurrentSelectedProjectPath = FText::GetEmpty();
+		CurrentSelectedProjectPath = FString();
 	}
 
 	ProjectSelectionChangedDelegate.ExecuteIfBound(ProjectFile);
 }
 
 
-FReply SProjectBrowser::HandleMarketplaceTabButtonClicked()
+FReply SProjectBrowser::OnOpenMarketplace()
 {
 	ILauncherPlatform* LauncherPlatform = FLauncherPlatformModule::Get();
 
@@ -1442,7 +1296,7 @@ void SProjectBrowser::OnFilterTextChanged(const FText& InText)
 {
 	ProjectItemFilter.SetRawFilterText(InText);
 	SearchBoxPtr->SetError(ProjectItemFilter.GetFilterErrorText());
-	PopulateFilteredProjectCategories();
+	PopulateFilteredProjects();
 }
 
 void SProjectBrowser::OnAutoloadLastProjectChanged(ECheckBoxState NewState)
@@ -1460,26 +1314,17 @@ void SProjectBrowser::OnAutoloadLastProjectChanged(ECheckBoxState NewState)
 
 EVisibility SProjectBrowser::GetProjectCategoryVisibility(const TSharedRef<FProjectCategory> InCategory) const
 {
-	if (NumFilteredProjects == 0)
-	{
-		return EVisibility::Collapsed;
-	}
-	return InCategory->FilteredProjectItemsSource.Num() ? EVisibility::Visible : EVisibility::Collapsed;
+	return FilteredProjectItemsSource.Num() ? EVisibility::Visible : EVisibility::Collapsed;
 }
 
 EVisibility SProjectBrowser::GetNoProjectsErrorVisibility() const
 {
-	return bHasProjectFiles ? EVisibility::Collapsed : EVisibility::Visible;
+	return ProjectItemsSource.Num() > 0 ? EVisibility::Collapsed : EVisibility::Visible;
 }
 
 EVisibility SProjectBrowser::GetNoProjectsAfterFilterErrorVisibility() const
 {
-	return (bHasProjectFiles && NumFilteredProjects == 0) ? EVisibility::Visible : EVisibility::Collapsed;
-}
-
-EVisibility SProjectBrowser::GetFilterActiveOverlayVisibility() const
-{
-	return ProjectItemFilter.GetRawFilterText().IsEmpty() ? EVisibility::Collapsed : EVisibility::HitTestInvisible;
+	return (ProjectItemsSource.Num() && FilteredProjectItemsSource.Num() == 0) ? EVisibility::Visible : EVisibility::Collapsed;
 }
 
 FText SProjectBrowser:: GetItemHighlightText() const
@@ -1487,76 +1332,4 @@ FText SProjectBrowser:: GetItemHighlightText() const
 	return ProjectItemFilter.GetRawFilterText();
 }
 
-
-SRecentProjectBrowser::SRecentProjectBrowser()
-{
-	
-}
-
-void SRecentProjectBrowser::Construct(const FArguments& InArgs)
-{
-	CategoriesBox = SNew(SVerticalBox);
-
-	const FString EngineIdentifier = FDesktopPlatformModule::Get()->GetCurrentEngineIdentifier();
-
-	TSharedRef<FProjectCategory> NewCategory = MakeShareable(new FProjectCategory);
-	NewCategory->CategoryName = LOCTEXT("ProjectDialog_Recent", "Recent Projects");
-	ProjectCategories.Add(NewCategory);
-
-	TArray<FString> RecentProjects = GetDefault<UEditorSettings>()->RecentlyOpenedProjectFiles;
-	for (int32 Idx = 0; Idx < InArgs._NumProjects && Idx < RecentProjects.Num(); ++Idx)
-	{
-		const FString ProjectFilename = RecentProjects[Idx];
-		
-		TSharedPtr<FProjectItem> ProjectItem = CreateProjectItem(ProjectFilename);
-
-		if (ProjectItem)
-		{
-			AddProjectToCategory(ProjectItem.ToSharedRef(), LOCTEXT("ProjectDialog_Recent", "Recent Projects"));
-		}
-	}
-
-	check(ProjectCategories.Num() == 1);
-
-	PopulateFilteredProjectCategories();
-
-	ConstructCategory(CategoriesBox.ToSharedRef(), ProjectCategories[0]);
-
-	ChildSlot
-	[
-		SNew(SBorder)
-		.BorderImage(FEditorStyle::GetBrush("ToolPanel.GroupBorder"))
-		.Padding(FMargin(8.f, 4.f))
-		[
-			SNew(SScrollBox)
-			+ SScrollBox::Slot()
-			[
-				CategoriesBox.ToSharedRef()
-			]
-		]
-	];
-
-	if (ProjectCategories.Num() > 0)
-	{
-		const TSharedRef<FProjectCategory> Category = ProjectCategories[0];
-
-		if (ensure(Category->ProjectTileView.IsValid()) && Category->ProjectItemsSource.Num() > 0)
-		{
-			Category->ProjectTileView->SetSelection(Category->ProjectItemsSource[0], ESelectInfo::Direct);
-		}
-	}
-
-	bHasProjectFiles = false;
-	for (const TSharedRef<FProjectCategory>& Category : ProjectCategories)
-	{
-		if (Category->ProjectItemsSource.Num() > 0)
-		{
-			bHasProjectFiles = true;
-			break;
-		}
-	}
-
-	ProjectSelectionChangedDelegate = InArgs._OnSelectionChanged;
-}
-
 #undef LOCTEXT_NAMESPACE
diff --git a/Engine/Source/Editor/GameProjectGeneration/Private/SProjectBrowser.h b/Engine/Source/Editor/GameProjectGeneration/Private/SProjectBrowser.h
--- a/Engine/Source/Editor/GameProjectGeneration/Private/SProjectBrowser.h
+++ b/Engine/Source/Editor/GameProjectGeneration/Private/SProjectBrowser.h
@@ -14,9 +14,14 @@ class ITableRow;
 class STableViewBase;
 class SToolTip;
 class SVerticalBox;
+
+template<typename T>
+class STileView;
+
 struct FProjectCategory;
 struct FProjectItem;
 struct FSlateBrush;
+
 enum class ECheckBoxState : uint8;
 
 DECLARE_DELEGATE_OneParam(FProjectSelectionChanged, FString);
@@ -31,14 +36,9 @@ public:
 
 	DECLARE_DELEGATE(FNewProjectScreenRequested)
 
-	SLATE_BEGIN_ARGS(SProjectBrowser) :
-		_HideOpenButton(false)
+	SLATE_BEGIN_ARGS(SProjectBrowser)
 	{}
-
 		SLATE_ARGUMENT(FProjectSelectionChanged, OnSelectionChanged);
-
-		SLATE_ARGUMENT(bool, HideOpenButton)
-
 	SLATE_END_ARGS()
 
 public:
@@ -55,16 +55,26 @@ public:
 
 	bool HasProjects() const;
 
-	void ClearSelection();
+	bool HasSelectedProjectFile() const { return !GetSelectedProjectFile().IsEmpty(); }
 
 	FString GetSelectedProjectFile() const;
 
+	FText GetSelectedProjectName() const;
+
 	/** Begins the opening process for the selected project */
 	void OpenSelectedProject();
 
-protected:
+	FReply OnBrowseToProject();
 
-	void ConstructCategory( const TSharedRef<SVerticalBox>& CategoriesBox, const TSharedRef<FProjectCategory>& Category );
+	/** Callback for clicking the 'Marketplace' button. */
+	FReply OnOpenMarketplace();
+
+	/** Handler for when the Open button is clicked */
+	FReply OnOpenProject();
+
+	/** Whether to autoload the last project. */
+	void OnAutoloadLastProjectChanged(ECheckBoxState NewState);
+protected:
 
 	/** Creates a row in the template list */
 	TSharedRef<ITableRow> MakeProjectViewWidget( TSharedPtr<FProjectItem> ProjectItem, const TSharedRef<STableViewBase>& OwnerTable );
@@ -90,20 +100,14 @@ protected:
 	/** Gets the currently selected template item */
 	TSharedPtr<FProjectItem> GetSelectedProjectItem( ) const;
 
-	/** Gets the label to show the currently selected template */
-	FText GetSelectedProjectName( ) const;
-
 	/** Populates ProjectItemsSource with projects found on disk */
 	FReply FindProjects( );
 
-	/** Adds the specified project to the specified category. Creates a new category if necessary. */
-	void AddProjectToCategory( const TSharedRef<FProjectItem>& ProjectItem, const FText& ProjectCategory );
-
 	/** Opens the specified project file */
 	bool OpenProject( const FString& ProjectFile );
 
 	/** Populate the list of filtered project categories */
-	void PopulateFilteredProjectCategories();
+	void PopulateFilteredProjects();
 
 	/**
 	 * Called after a key is pressed when this widget has focus (this event bubbles if not handled)
@@ -115,23 +119,9 @@ protected:
 	 */
 	virtual FReply OnKeyDown( const FGeometry& MyGeometry, const FKeyEvent& InKeyEvent ) override;
 
-protected:
-
-	/** Holds the collection of project categories. */
-	TArray<TSharedRef<FProjectCategory> > ProjectCategories;
-
-	bool bHasProjectFiles;
-
-	TSharedPtr<SVerticalBox> CategoriesBox;
-
-	FProjectSelectionChanged ProjectSelectionChangedDelegate;
-
 private:
 
-	FReply OnBrowseToProjectClicked();
 
-	/** Handler for when the Open button is clicked */
-	FReply HandleOpenProjectButtonClicked( );
 
 	/** Returns true if the user is allowed to open a project with the supplied name and path */
 	bool HandleOpenProjectButtonIsEnabled( ) const;
@@ -140,16 +130,12 @@ private:
 	void HandleProjectItemDoubleClick( TSharedPtr<FProjectItem> ProjectItem );
 
 	/** Handler for when the selection changes in the project view */
-	void HandleProjectViewSelectionChanged( TSharedPtr<FProjectItem> ProjectItem, ESelectInfo::Type SelectInfo, FText CategoryName );
+	void HandleProjectViewSelectionChanged( TSharedPtr<FProjectItem> ProjectItem, ESelectInfo::Type SelectInfo);
 
-	/** Callback for clicking the 'Marketplace' button. */
-	FReply HandleMarketplaceTabButtonClicked( );
 
 	/** Called when the text in the filter box is changed */
 	void OnFilterTextChanged(const FText& InText);
 
-	/** Whether to autoload the last project. */
-	void OnAutoloadLastProjectChanged(ECheckBoxState NewState);
 
 	/** Get the visibility of the specified category */
 	EVisibility GetProjectCategoryVisibility(const TSharedRef<FProjectCategory> InCategory) const;
@@ -159,13 +145,16 @@ private:
 	/** Get the visibility of the "no projects" error */
 	EVisibility GetNoProjectsAfterFilterErrorVisibility() const;
 	
-	/** Get the visibility of the active search overlay */
-	EVisibility GetFilterActiveOverlayVisibility() const;
-
 	/** Get the filter text to highlight on items in the list */
 	FText GetItemHighlightText() const;
 
 private:
+	TSharedPtr<STileView<TSharedPtr<FProjectItem>>> ProjectTileView;
+	TArray<TSharedPtr<FProjectItem>> ProjectItemsSource;
+	TArray<TSharedPtr<FProjectItem>> FilteredProjectItemsSource;
+	TSharedPtr<SVerticalBox> ProjectsBox;
+
+	FProjectSelectionChanged ProjectSelectionChangedDelegate;
 
 	/** Search box used to set the filter text */
 	TSharedPtr<class SSearchBox> SearchBoxPtr;
@@ -173,14 +162,10 @@ private:
 	/** Filter that is used to test for the visibility of projects */
 	typedef TTextFilter<const TSharedPtr<FProjectItem>> ProjectItemTextFilter;
 	ProjectItemTextFilter ProjectItemFilter;
-	int32 NumFilteredProjects;
 
-	int32 ThumbnailBorderPadding;
-	int32 ThumbnailSize;
 	bool bPreventSelectionChangeEvent;
 
-	TSharedPtr<FProjectItem> CurrentlySelectedItem;
-	FText CurrentSelectedProjectPath;
+	FString CurrentSelectedProjectPath;
 
 	bool IsOnlineContentFinished;
 
@@ -188,20 +173,3 @@ private:
 	FNewProjectScreenRequested NewProjectScreenRequestedDelegate;
 };
 
-/** Class to only display a finite number of most-recent projects. */
-class SRecentProjectBrowser : public SProjectBrowser
-{
-	SLATE_BEGIN_ARGS(SRecentProjectBrowser) 
-	{
-		_NumProjects = 10;
-	}
-
-	SLATE_ARGUMENT(int32, NumProjects);
-	SLATE_EVENT(FProjectSelectionChanged, OnSelectionChanged);
-
-	SLATE_END_ARGS()
-
-	SRecentProjectBrowser();
-
-	void Construct(const FArguments& InArgs);
-};
\ No newline at end of file
diff --git a/Engine/Source/Editor/GameProjectGeneration/Private/SProjectDialog.cpp b/Engine/Source/Editor/GameProjectGeneration/Private/SProjectDialog.cpp
new file mode 100644
--- /dev/null
+++ b/Engine/Source/Editor/GameProjectGeneration/Private/SProjectDialog.cpp
@@ -0,0 +1,2067 @@
+// Copyright Epic Games, Inc. All Rights Reserved.
+
+#include "SProjectDialog.h"
+#include "Widgets/Layout/SBorder.h"
+#include "Widgets/SBoxPanel.h"
+#include "GameProjectGenerationModule.h"
+#include "TemplateCategory.h"
+#include "SProjectBrowser.h"
+#include "Widgets/Layout/SSeparator.h"
+#include "IDocumentation.h"
+#include "IDesktopPlatform.h"
+#include "DesktopPlatformModule.h"
+#include "TemplateItem.h"
+#include "SourceCodeNavigation.h"
+#include "GameProjectUtils.h"
+#include "Framework/Application/SlateApplication.h"
+#include "Editor.h"
+#include "Internationalization/BreakIterator.h"
+#include "Settings/EditorSettings.h"
+#include "Widgets/Layout/SScrollBorder.h"
+#include "Widgets/Layout/SScrollBox.h"
+#include "HardwareTargetingModule.h"
+#include "Widgets/Input/SSegmentedControl.h"
+#include "GameProjectGenerationLog.h"
+#include "Interfaces/IPluginManager.h"
+#include "HAL/PlatformFileManager.h"
+#include "Dialogs/SOutputLogDialog.h"
+#include "Misc/MessageDialog.h"
+#include "HAL/FileManager.h"
+#include "ProjectDescriptor.h"
+#include "Widgets/Images/SImage.h"
+#include "Widgets/SToolTip.h"
+#include "Widgets/Layout/SGridPanel.h"
+#include "Widgets/SOverlay.h"
+#include "Widgets/Input/SEditableTextBox.h"
+#include "Widgets/Input/SButton.h"
+#include "SWarningOrErrorBox.h"
+#include "SGetSuggestedIDEWidget.h"
+#include "Brushes/SlateDynamicImageBrush.h"
+#include "Widgets/Layout/SWidgetSwitcher.h"
+#include "SProjectDialog.h"
+#include "LauncherPlatformModule.h"
+
+
+#define LOCTEXT_NAMESPACE "GameProjectGeneration"
+
+namespace NewProjectDialogDefs
+{
+	constexpr float MajorItemWidth = 304;
+	constexpr float MajorItemHeight = 104;
+	constexpr float TemplateTileHeight = 153;
+	constexpr float TemplateTileWidth = 102;
+	constexpr float ThumbnailSize = 64.0f, ThumbnailPadding = 5.f;
+	const FName DefaultCategoryName = "Games";
+}
+
+class SMajorCategoryTile : public SCompoundWidget
+{
+	SLATE_BEGIN_ARGS(SMajorCategoryTile) {}
+		SLATE_ATTRIBUTE(bool, IsSelected)
+	SLATE_END_ARGS()
+
+public:
+	void Construct(const FArguments& InArgs, TSharedPtr<FTemplateCategory> Item)
+	{
+		IsSelected = InArgs._IsSelected;
+
+		ChildSlot
+		[
+			SNew(SOverlay)
+			+ SOverlay::Slot()
+			[
+				SNew(SImage)
+				.Image(Item->Icon)
+			]
+			+ SOverlay::Slot()
+			.HAlign(HAlign_Left)
+			.VAlign(VAlign_Bottom)
+			.Padding(FMargin(18.0f, 8.0f))
+			[
+				SNew(STextBlock)
+				.Font(FAppStyle::Get().GetFontStyle("HeadingExtraSmall"))
+				.ColorAndOpacity(FLinearColor(1, 1, 1, .9f))
+				.TransformPolicy(ETextTransformPolicy::ToUpper)
+				.ShadowOffset(FVector2D(1, 1))
+				.ShadowColorAndOpacity(FLinearColor(0, 0, 0, .75))
+				.Text(Item->DisplayName)
+				.WrapTextAt(250.0f)
+			]
+			+ SOverlay::Slot()
+			[
+				SNew(SImage)
+				.Visibility(EVisibility::HitTestInvisible)
+				.Image(this, &SMajorCategoryTile::GetSelectionOutlineBrush)
+			]
+		];
+	}
+
+private:
+	const FSlateBrush* GetSelectionOutlineBrush() const
+	{
+		const bool bIsSelected = IsSelected.Get();
+		const bool bIsTileHovered = IsHovered();
+
+		if (bIsSelected && bIsTileHovered)
+		{
+			static const FName SelectedHover("ProjectBrowser.ProjectTile.SelectedHoverBorder");
+			return FAppStyle::Get().GetBrush(SelectedHover);
+		}
+		else if (bIsSelected)
+		{
+			static const FName Selected("ProjectBrowser.ProjectTile.SelectedBorder");
+			return FAppStyle::Get().GetBrush(Selected);
+		}
+		else if (bIsTileHovered)
+		{
+			static const FName Hovered("ProjectBrowser.ProjectTile.HoverBorder");
+			return FAppStyle::Get().GetBrush(Hovered);
+		}
+
+		return FStyleDefaults::GetNoBrush();
+	}
+private:
+	TAttribute<bool> IsSelected;
+};
+
+
+/** Slate tile widget for template projects */
+class STemplateTile : public STableRow<TSharedPtr<FTemplateItem>>
+{
+public:
+	SLATE_BEGIN_ARGS( STemplateTile ){}
+		SLATE_ARGUMENT(TSharedPtr<FTemplateItem>, Item)
+	SLATE_END_ARGS()
+
+private:
+	TWeakPtr<FTemplateItem> Item;
+
+public:
+	/** Static build function */
+	static TSharedRef<ITableRow> BuildTile(TSharedPtr<FTemplateItem> Item, const TSharedRef<STableViewBase>& OwnerTable)
+	{
+		if (!ensure(Item.IsValid()))
+		{
+			return SNew(STableRow<TSharedPtr<FTemplateItem>>, OwnerTable);
+		}
+
+		return SNew(STemplateTile, OwnerTable).Item(Item);
+	}
+
+	/** Constructs this widget with InArgs */
+	void Construct( const FArguments& InArgs, const TSharedRef<STableViewBase>& OwnerTable )
+	{
+		check(InArgs._Item.IsValid())
+		Item = InArgs._Item;
+
+		STableRow::Construct(
+			STableRow::FArguments()
+			.Style(FAppStyle::Get(), "ProjectBrowser.TableRow")
+			.Padding(2.0f)
+			.Content()
+			[
+				SNew(SBorder)
+				.Padding(FMargin(0.0f, 0.0f, 5.0f, 5.0f))
+				.BorderImage(FAppStyle::Get().GetBrush("ProjectBrowser.ProjectTile.DropShadow"))
+				[
+					SNew(SOverlay)
+					+ SOverlay::Slot()
+					[
+						SNew(SVerticalBox)
+						// Thumbnail
+						+ SVerticalBox::Slot()
+						.AutoHeight()
+						.HAlign(HAlign_Center)
+						.VAlign(VAlign_Center)
+						[
+							SNew(SBox)
+							.WidthOverride(102)
+							.HeightOverride(102)
+							[
+								SNew(SBorder)
+								.Padding(0)
+								.BorderImage(FAppStyle::Get().GetBrush("ProjectBrowser.ProjectTile.ThumbnailAreaBackground"))
+								.HAlign(HAlign_Center)
+								.VAlign(VAlign_Center)
+								[
+									SNew(SImage)
+									.Image(this, &STemplateTile::GetThumbnail)
+									.DesiredSizeOverride(InArgs._Item->bThumbnailAsIcon ? TOptional<FVector2D>() : FVector2D(NewProjectDialogDefs::ThumbnailSize, NewProjectDialogDefs::ThumbnailSize))
+								]
+							]
+						]
+						// Name
+						+ SVerticalBox::Slot()
+						[
+							SNew(SBorder)
+							.Padding(FMargin(NewProjectDialogDefs::ThumbnailPadding, 0))
+							.VAlign(VAlign_Top)
+							.Padding(FMargin(3.0f, 3.0f))
+							.BorderImage(FAppStyle::Get().GetBrush("ProjectBrowser.ProjectTile.NameAreaBackground"))
+							[
+								SNew(STextBlock)
+								.Font(FAppStyle::Get().GetFontStyle("ProjectBrowser.ProjectTile.Font"))
+								.WrapTextAt(NewProjectDialogDefs::TemplateTileWidth-4.0f)
+								.LineBreakPolicy(FBreakIterator::CreateCamelCaseBreakIterator())
+								.Text(InArgs._Item->Name)
+								.ColorAndOpacity(FAppStyle::Get().GetSlateColor("Colors.Foreground"))
+							]
+						]
+					]
+					+ SOverlay::Slot()
+					[
+						SNew(SImage)
+						.Visibility(EVisibility::HitTestInvisible)
+						.Image_Lambda
+						(
+							[this]()
+							{
+								const bool bIsSelected = IsSelected();
+								const bool bIsRowHovered = IsHovered();
+
+								if (bIsSelected && bIsRowHovered)
+								{
+									static const FName SelectedHover("ProjectBrowser.ProjectTile.SelectedHoverBorder");
+									return FAppStyle::Get().GetBrush(SelectedHover);
+								}
+								else if (bIsSelected)
+								{
+									static const FName Selected("ProjectBrowser.ProjectTile.SelectedBorder");
+									return FAppStyle::Get().GetBrush(Selected);
+								}
+								else if (bIsRowHovered)
+								{
+									static const FName Hovered("ProjectBrowser.ProjectTile.HoverBorder");
+									return FAppStyle::Get().GetBrush(Hovered);
+								}
+
+								return FStyleDefaults::GetNoBrush();
+							}
+						)
+					]
+				]
+			],
+			OwnerTable
+		);
+	}
+
+private:
+
+	/** Get this item's thumbnail or return the default */
+	const FSlateBrush* GetThumbnail() const
+	{
+		TSharedPtr<FTemplateItem> ItemPtr = Item.Pin();
+		if (ItemPtr.IsValid() && ItemPtr->Thumbnail.IsValid())
+		{
+			return ItemPtr->Thumbnail.Get();
+		}
+		return FEditorStyle::GetBrush("GameProjectDialog.DefaultGameThumbnail.Small");
+	}
+	
+};
+
+/**
+ * Simple widget used to display a folder path, and a name of a file
+ */
+class SFilepath : public SCompoundWidget
+{
+public:
+
+	SLATE_BEGIN_ARGS( SFilepath )
+		: _IsReadOnly(false)
+	{}
+		/** Attribute specifying the text to display in the folder input */
+		SLATE_ATTRIBUTE(FText, FolderPath)
+
+		/** Attribute specifying the text to display in the name input */
+		SLATE_ATTRIBUTE(FText, Name)
+
+		SLATE_ATTRIBUTE(FText, WarningText)
+
+		SLATE_ARGUMENT(bool, IsReadOnly)
+
+		/** Event that is triggered when the browser for folder button is clicked */
+		SLATE_EVENT(FOnClicked, OnBrowseForFolder)
+
+		/** Events for when the name field is manipulated */
+		SLATE_EVENT(FOnTextChanged, OnNameChanged)
+		SLATE_EVENT(FOnTextCommitted, OnNameCommitted)
+		
+		/** Events for when the folder field is manipulated */
+		SLATE_EVENT(FOnTextChanged, OnFolderChanged)
+		SLATE_EVENT(FOnTextCommitted, OnFolderCommitted)
+
+	SLATE_END_ARGS()
+
+	/** Constructs this widget with InArgs */
+	void Construct( const FArguments& InArgs )
+	{
+		WarningText = InArgs._WarningText;
+
+		ChildSlot
+		[
+			SNew(SHorizontalBox)
+			+ SHorizontalBox::Slot()
+			.HAlign(HAlign_Left)
+			.VAlign(VAlign_Top)
+			.Padding(0.0f, 4.0f, 8.0f, 8.0f)
+			.AutoWidth()
+			[
+				SNew(STextBlock)
+				.Text(LOCTEXT("ProjectLocation", "Project Location"))
+			]
+			+ SHorizontalBox::Slot()
+			.VAlign(VAlign_Top)
+			.AutoWidth()
+			[
+				SNew(SBox)
+				.WidthOverride(595.0f)
+				[
+					SNew(SVerticalBox)
+					+ SVerticalBox::Slot()
+					.AutoHeight()
+					[
+						SNew(SEditableTextBox)
+						.IsReadOnly(InArgs._IsReadOnly)
+						.Text(InArgs._FolderPath)
+						.OnTextChanged(InArgs._OnFolderChanged)
+						.OnTextCommitted(InArgs._OnFolderCommitted)
+					]
+					+ SVerticalBox::Slot()
+					.Padding(0.0f, 8.0f)
+					[	
+						SNew(SWarningOrErrorBox)
+						.Visibility(this, &SFilepath::GetWarningVisibility)
+						.IconSize(FVector2D(16,16))
+						.Padding(FMargin(8.0f, 4.0f, 4.0f, 4.0f))
+						.Message(WarningText)
+						.MessageStyle(EMessageStyle::Error)
+					]
+				]
+			]
+			+ SHorizontalBox::Slot()
+			.VAlign(VAlign_Top)
+			.Padding(2.0f, 2.0f, 0.0f, 0.0f)
+			.AutoWidth()
+			[
+				SNew(SButton)
+				.ButtonStyle(FAppStyle::Get(), "SimpleButton")
+				.OnClicked(InArgs._OnBrowseForFolder)
+				.ToolTipText(LOCTEXT("BrowseForFolder", "Browse for a folder"))
+				.Visibility(InArgs._IsReadOnly ? EVisibility::Hidden : EVisibility::Visible)
+				.ContentPadding(0)
+				[
+					SNew(SImage)
+					.Image(FAppStyle::Get().GetBrush("Icons.FolderClosed"))
+					.ColorAndOpacity(FSlateColor::UseForeground())
+				]
+			]
+			+ SHorizontalBox::Slot()
+			.HAlign(HAlign_Right)
+			[
+				SNew(SHorizontalBox)
+				+SHorizontalBox::Slot()
+				.VAlign(VAlign_Top)
+				.AutoWidth()
+				.Padding(32.0f, 4.0f, 8.0f, 8.0f)
+				[
+					SNew(STextBlock)
+					.Text(LOCTEXT("ProjectName", "Project Name"))
+				]
+				+ SHorizontalBox::Slot()
+				.VAlign(VAlign_Top)
+				[
+					SNew(SBox)
+					.WidthOverride(275.0f)
+					[
+						SNew(SEditableTextBox)
+						.IsReadOnly(InArgs._IsReadOnly)
+						.Text(InArgs._Name)
+						.OnTextChanged(InArgs._OnNameChanged)
+						.OnTextCommitted(InArgs._OnNameCommitted)
+					]
+				]
+			]
+		];
+	}
+
+private:
+	EVisibility GetWarningVisibility() const
+	{
+		return WarningText.Get().IsEmpty() ? EVisibility::Hidden : EVisibility::HitTestInvisible;
+	}
+private:
+	TAttribute<FText> WarningText;
+};
+
+void SProjectDialog::Construct(const FArguments& InArgs, EProjectDialogModeMode Mode)
+{
+	bLastGlobalValidityCheckSuccessful = true;
+	bLastNameAndLocationValidityCheckSuccessful = true;
+
+	PopulateTemplateCategories();
+
+	ProjectBrowser = SNew(SProjectBrowser);
+
+	ChildSlot
+	[
+		SNew(SBorder)
+		.BorderImage(FAppStyle::Get().GetBrush("Brushes.Panel"))
+		.Padding(FMargin(16.0f, 8.0f))
+		[
+			SNew(SVerticalBox)
+			+ SVerticalBox::Slot()
+			.Padding(0.0f, 0.0f, 0.0f, 12.0f)
+			[
+				MakeHybridView(Mode)
+			]
+			+SVerticalBox::Slot()
+			.AutoHeight()
+			.Padding(-16.0f, 0.0f)
+			[
+				SNew(SSeparator)
+				.Orientation(EOrientation::Orient_Horizontal)
+				.Thickness(2.0f)
+			]
+			+SVerticalBox::Slot()
+			.AutoHeight()
+			[
+				SNew(SBox)
+				.HeightOverride(145.0f)
+				[
+					SAssignNew(PathAreaSwitcher, SWidgetSwitcher)
+					+SWidgetSwitcher::Slot()
+					[
+						MakeNewProjectPathArea()
+					]
+					+SWidgetSwitcher::Slot()
+					[
+						MakeOpenProjectPathArea()
+					]
+				]
+			]
+		]
+	];
+
+	RegisterActiveTimer(1.0f, FWidgetActiveTimerDelegate::CreateLambda(
+		[this](double, float)
+		{
+			UpdateProjectFileValidity();
+			return EActiveTimerReturnType::Continue;
+		}));
+
+	if (Mode == EProjectDialogModeMode::OpenProject)
+	{
+		PathAreaSwitcher->SetActiveWidgetIndex(1);
+	}
+	else if (Mode == EProjectDialogModeMode::Hybrid && ProjectBrowser->HasProjects())
+	{
+		// Select recent projects
+		OnRecentProjectsClicked();
+	}
+	else if(!ProjectBrowser->HasProjects() || Mode == EProjectDialogModeMode::NewProject)
+	{
+		// Select the first template category
+		MajorCategoryList->SetSelection(TemplateCategories[0]);
+	}
+
+}
+
+SProjectDialog::~SProjectDialog()
+{
+	// remove any UTemplateProjectDefs we were keeping alive
+	for (const TPair<FName, TArray<TSharedPtr<FTemplateItem>>>& Pair : Templates)
+	{
+		for (const TSharedPtr<FTemplateItem>& Template : Pair.Value)
+		{
+			if (Template->CodeTemplateDefs != nullptr)
+			{
+				Template->CodeTemplateDefs->RemoveFromRoot();
+			}
+
+			if (Template->BlueprintTemplateDefs != nullptr)
+			{
+				Template->BlueprintTemplateDefs->RemoveFromRoot();
+			}
+		}
+	}
+}
+
+void SProjectDialog::PopulateTemplateCategories()
+{
+	TemplateCategories.Empty();
+	CurrentCategory.Reset();
+
+	FGameProjectGenerationModule::Get().GetAllTemplateCategories(TemplateCategories);
+
+	if (TemplateCategories.Num() == 0)
+	{
+		TSharedPtr<FTemplateCategory> DefaultCategory = MakeShared<FTemplateCategory>();
+		static const FName DefaultCategoryKey("Default");
+		DefaultCategory->Key = DefaultCategoryKey;
+		DefaultCategory->DisplayName = LOCTEXT("ProjectDialog_DefaultCategoryName", "Blank Project");
+		DefaultCategory->Description = LOCTEXT("ProjectDialog_DefaultCategoryDescription", "Create a new blank Unreal project.");
+		DefaultCategory->Icon = FAppStyle::Get().GetBrush("GameProjectDialog.DefaultGameThumbnail");
+
+		TemplateCategories.Add(DefaultCategory);
+	}
+}
+
+TSharedRef<SWidget> SProjectDialog::MakeNewProjectDialogButtons()
+{
+	return
+		SNew(SHorizontalBox)
+		+ SHorizontalBox::Slot()
+		.Padding(8.0f, 0.0f, 8.0f, 0.0f)
+		.VAlign(VAlign_Center)
+		.AutoWidth()
+		[
+			SNew(SGetSuggestedIDEWidget)
+			.VisibilityOverride(this, &SProjectDialog::GetSuggestedIDEButtonVisibility)
+		]
+		+SHorizontalBox::Slot()
+		.Padding(8.0f, 0.0f, 8.0f, 0.0f)
+		.VAlign(VAlign_Center)
+		.AutoWidth()
+		[
+			SNew(SButton)
+			.Visibility(this, &SProjectDialog::GetCreateButtonVisibility)
+			.HAlign(HAlign_Center)
+			.VAlign(VAlign_Center)
+			.ButtonStyle(FAppStyle::Get(), "PrimaryButton")
+			.TextStyle(FAppStyle::Get(), "DialogButtonText")
+			.Text(LOCTEXT("CreateNewProject", "Create"))
+			.IsEnabled(this, &SProjectDialog::CanCreateProject)
+			.OnClicked_Lambda([this](){CreateAndOpenProject(); return FReply::Handled(); })
+		]
+		+ SHorizontalBox::Slot()
+		.Padding(8.0f, 0.0f, 0.0f, 0.0f)
+		.VAlign(VAlign_Center)
+		.AutoWidth()
+		[
+			SNew(SButton)
+			.HAlign(HAlign_Center)
+			.VAlign(VAlign_Center)
+			//.ButtonStyle(FAppStyle::Get(), "SecondaryButton")
+			.TextStyle(FAppStyle::Get(), "DialogButtonText")
+			.Text(LOCTEXT("CancelNewProjectCreation", "Cancel"))
+			.OnClicked(this, &SProjectDialog::OnCancel)
+		];
+}
+
+TSharedRef<SWidget> SProjectDialog::MakeOpenProjectDialogButtons()
+{
+	TSharedRef<SProjectBrowser> ProjectBrowserRef = ProjectBrowser.ToSharedRef();
+
+	return
+		SNew(SHorizontalBox)
+		+ SHorizontalBox::Slot()
+		.Padding(8.0f, 0.0f, 16.0f, 0.0f)
+		.VAlign(VAlign_Center)
+		.AutoWidth()
+		[
+			SNew(SButton)
+			.HAlign(HAlign_Center)
+			.VAlign(VAlign_Center)
+			.ToolTipText(LOCTEXT("MarketplaceToolTip", "Check out the Marketplace to find new projects!"))
+			.Visibility(FLauncherPlatformModule::Get()->CanOpenLauncher(true) ? EVisibility::Visible : EVisibility::Collapsed)
+			.OnClicked(ProjectBrowserRef, &SProjectBrowser::OnOpenMarketplace)
+			.ContentPadding(FMargin(4.0f, 2.0f,4.0f, 1.0f))
+			[	
+				SNew(SHorizontalBox)
+				+SHorizontalBox::Slot()
+				.Padding(4.0f,0.0f)
+				.AutoWidth()
+				.HAlign(HAlign_Left)
+				.VAlign(VAlign_Center)
+				[
+					SNew(SImage)
+					.Image(FAppStyle::Get().GetBrush("LevelEditor.OpenMarketplace"))
+					.DesiredSizeOverride(FVector2D(16,16))
+					.ColorAndOpacity(FSlateColor::UseForeground())
+				]
+				+ SHorizontalBox::Slot()
+				.Padding(4.0f, 0.0f)
+				.VAlign(VAlign_Center)
+				[
+					SNew(STextBlock)
+					.TextStyle(FAppStyle::Get(), "DialogButtonText")
+					.Text(LOCTEXT("OpenMarketplace", "Marketplace..."))
+				]
+			]
+		]
+		+ SHorizontalBox::Slot()
+		.Padding(8.0f, 0.0f)
+		.VAlign(VAlign_Center)
+		.AutoWidth()
+		[
+			SNew(SButton)
+			.HAlign(HAlign_Center)
+			.VAlign(VAlign_Center)
+			//.ButtonStyle(FAppStyle::Get(), "SecondaryButton")
+			.TextStyle(FAppStyle::Get(), "DialogButtonText")
+			.Text(LOCTEXT("BrowseForProjects", "Browse..."))
+			.OnClicked(ProjectBrowserRef, &SProjectBrowser::OnBrowseToProject)
+		]
+		+SHorizontalBox::Slot()
+		.Padding(8.0f, 0.0f)
+		.VAlign(VAlign_Center)
+		.AutoWidth()
+		[
+			SNew(SButton)
+			.Visibility(this, &SProjectDialog::GetCreateButtonVisibility)
+			.HAlign(HAlign_Center)
+			.VAlign(VAlign_Center)
+			.ButtonStyle(FAppStyle::Get(), "PrimaryButton")
+			.TextStyle(FAppStyle::Get(), "DialogButtonText")
+			.Text(LOCTEXT("OpenProject", "Open"))
+			.IsEnabled(ProjectBrowserRef, &SProjectBrowser::HasSelectedProjectFile)
+			.OnClicked(ProjectBrowserRef, &SProjectBrowser::OnOpenProject)
+		]
+		+ SHorizontalBox::Slot()
+		.Padding(8.0f, 0.0f, 0.0f, 0.0f)
+		.VAlign(VAlign_Center)
+		.AutoWidth()
+		[
+			SNew(SButton)
+			.HAlign(HAlign_Center)
+			.VAlign(VAlign_Center)
+			//.ButtonStyle(FAppStyle::Get(), "SecondaryButton")
+			.TextStyle(FAppStyle::Get(), "DialogButtonText")
+			.Text(LOCTEXT("CancelNewProjectCreation", "Cancel"))
+			.OnClicked(this, &SProjectDialog::OnCancel)
+		];
+}
+
+TSharedRef<SWidget> SProjectDialog::MakeTemplateProjectView()
+{
+	return
+		SNew(SVerticalBox)
+		// Templates list
+		+ SVerticalBox::Slot()
+		.FillHeight(1.0f)
+		.Padding(0)
+		[
+			SNew(SHorizontalBox)
+			+ SHorizontalBox::Slot()
+			.Padding(0.0f, 0.0f, 0.0f, -12.0f)
+			[
+				SNew(SScrollBorder, TemplateListView.ToSharedRef())
+				[
+					TemplateListView.ToSharedRef()
+				]
+			]						
+			+ SHorizontalBox::Slot()
+			.Padding(0, -8.0f, 0.0f, -12.0f)
+			.AutoWidth()
+			[
+				SNew(SSeparator)
+				.Orientation(EOrientation::Orient_Vertical)
+				.Thickness(2.0f)
+			]
+			// Selected template details
+			+ SHorizontalBox::Slot()
+			.Padding(8, 0.0f)
+			.AutoWidth()
+			[
+				SNew(SVerticalBox)
+				// Preview image
+				+ SVerticalBox::Slot()
+				.AutoHeight()
+				.Padding(FMargin(0.0f, 0.0f, 0.0f, 15.f))
+				[
+					SNew(SImage)
+					.DesiredSizeOverride(FVector2D(358,160))
+					.Image(this, &SProjectDialog::GetSelectedTemplatePreviewImage)
+				]
+				// Template Name
+				+ SVerticalBox::Slot()
+				.Padding(FMargin(0.0f, 0.0f, 0.0f, 10.0f))
+				.AutoHeight()
+				[
+					SNew(STextBlock)
+					.AutoWrapText(true)
+					.TextStyle(FEditorStyle::Get(), "DialogButtonText")
+					.Font(FAppStyle::Get().GetFontStyle("HeadingExtraSmall"))
+					.ColorAndOpacity(FAppStyle::Get().GetSlateColor("Colors.White"))
+					.Text(this, &SProjectDialog::GetSelectedTemplateProperty, &FTemplateItem::Name)
+				]
+				// Template Description
+				+ SVerticalBox::Slot()
+				.AutoHeight()
+				[
+					SNew(SBox)
+					.HeightOverride(120.0f)
+					.WidthOverride(358.0f)
+					[
+						SNew(SScrollBox)
+						+ SScrollBox::Slot()
+						[
+							SNew(SVerticalBox)
+							+SVerticalBox::Slot()
+							.Padding(FMargin(0.0f, 0.0f, 0.0f, 10.0f))
+							.AutoHeight()
+							[
+								SNew(STextBlock)
+								.WrapTextAt(350.0f)
+								.Text(this, &SProjectDialog::GetSelectedTemplateProperty, &FTemplateItem::Description)
+							]
+							// Asset types
+							+ SVerticalBox::Slot()
+							.AutoHeight()
+							.Padding(FMargin(0.0f, 5.0f, 0.0f, 5.0f))
+							[
+								SNew(SVerticalBox)
+								.Visibility(this, &SProjectDialog::GetSelectedTemplateAssetVisibility)
+								+ SVerticalBox::Slot()
+								[
+									SNew(STextBlock)
+									.TextStyle(FEditorStyle::Get(), "DialogButtonText")
+									.Text(LOCTEXT("ProjectTemplateAssetTypes", "Asset Type References"))
+								]
+								+ SVerticalBox::Slot()
+								.AutoHeight()
+								[
+									SNew(STextBlock)
+									.AutoWrapText(true)
+									.Text(this, &SProjectDialog::GetSelectedTemplateAssetTypes)
+								]
+							]
+							// Class types
+							+ SVerticalBox::Slot()
+							.AutoHeight()
+							.Padding(FMargin(0.0f, 5.0f, 0.0f, 5.0f))
+							[
+								SNew(SVerticalBox)
+								.Visibility(this, &SProjectDialog::GetSelectedTemplateClassVisibility)
+								+ SVerticalBox::Slot()
+								[
+									SNew(STextBlock)
+									.TextStyle(FEditorStyle::Get(), "DialogButtonText")
+									.Text(LOCTEXT("ProjectTemplateClassTypes", "Class Type References"))
+								]
+								+ SVerticalBox::Slot()
+								.AutoHeight()
+								[
+									SNew(STextBlock)
+									.AutoWrapText(true)
+									.Text(this, &SProjectDialog::GetSelectedTemplateClassTypes)
+								]
+							]
+						]
+					]
+				]
+				// Project Options
+				+ SVerticalBox::Slot()
+				.Expose(ProjectOptionsSlot)
+			]
+		];
+}
+
+TSharedRef<SWidget> SProjectDialog::MakeHybridView(EProjectDialogModeMode Mode)
+{
+	bCopyStarterContent = GEditor ? GetDefault<UEditorSettings>()->bCopyStarterContentPreference : true;
+
+	SelectedHardwareClassTarget = EHardwareClass::Desktop;
+	SelectedGraphicsPreset = EGraphicsPreset::Maximum;
+
+	// Find all template projects
+	Templates = FindTemplateProjects();
+	SetDefaultProjectLocation();
+
+	Mode = EProjectDialogModeMode::Hybrid;
+
+	TemplateListView = SNew(STileView<TSharedPtr<FTemplateItem>>)
+		.ListItemsSource(&FilteredTemplateList)
+		.SelectionMode(ESelectionMode::Single)
+		.ClearSelectionOnClick(false)
+		.ItemAlignment(EListItemAlignment::LeftAligned)
+		.OnGenerateTile_Static(&STemplateTile::BuildTile)
+		.ItemHeight(NewProjectDialogDefs::TemplateTileHeight+9)
+		.ItemWidth(NewProjectDialogDefs::TemplateTileWidth+9)
+		.OnSelectionChanged(this, &SProjectDialog::HandleTemplateListViewSelectionChanged);
+
+	TSharedRef<SWidget> HybridView = 
+		SNew(SHorizontalBox)
+		+ SHorizontalBox::Slot()
+		.AutoWidth()
+		[ 
+			SNew(SVerticalBox)
+			+SVerticalBox::Slot()
+			.Padding(0.0f, 0.0f, 0.0f, 16.0f)
+			.AutoHeight()
+			[
+				SNew(SBox)
+				.Visibility(Mode == EProjectDialogModeMode::Hybrid ? EVisibility::Visible : EVisibility::Collapsed)
+				[	
+					MakeRecentProjectsTile()
+				]
+			]
+			+ SVerticalBox::Slot()
+			.Padding(0.0f, -4.0f, 0.0f, 0.0f)
+			[
+				SNew(SBorder)
+				.Visibility(Mode == EProjectDialogModeMode::OpenProject ? EVisibility::Collapsed : EVisibility::Visible)
+				.BorderImage(FAppStyle::Get().GetBrush("ProjectBrowser.MajorCategoryViewBorder"))
+				[
+					SAssignNew(MajorCategoryList, STileView<TSharedPtr<FTemplateCategory>>)
+					.ListItemsSource(&TemplateCategories)
+					.SelectionMode(ESelectionMode::Single)
+					.ClearSelectionOnClick(false)
+					.OnGenerateTile(this, &SProjectDialog::ConstructMajorCategoryTableRow)
+					.ItemHeight(NewProjectDialogDefs::MajorItemHeight)
+					.ItemWidth(NewProjectDialogDefs::MajorItemWidth)
+					.OnSelectionChanged(this, &SProjectDialog::OnMajorTemplateCategorySelectionChanged)
+				]
+			]
+		]
+		+SHorizontalBox::Slot()
+		.Padding(11.0f, 0.0f, 0.0f, 0.0f)
+		[
+			SAssignNew(TemplateAndRecentProjectsSwitcher, SWidgetSwitcher)
+			+ SWidgetSwitcher::Slot()
+			[
+				MakeTemplateProjectView()
+			]
+			+ SWidgetSwitcher::Slot()
+			[
+				ProjectBrowser.ToSharedRef()
+			]
+		];
+	
+	SetCurrentMajorCategory(ActiveCategory);
+
+	if (Mode == EProjectDialogModeMode::OpenProject)
+	{
+		TemplateAndRecentProjectsSwitcher->SetActiveWidgetIndex(1);
+	}
+
+	UpdateProjectFileValidity();
+
+	return HybridView;
+}
+
+TSharedRef<SWidget> SProjectDialog::MakeProjectOptionsWidget()
+{
+	IHardwareTargetingModule& HardwareTargeting = IHardwareTargetingModule::Get();
+
+	TSharedPtr<SVerticalBox> ProjectOptionsBox;
+
+	TSharedRef<SWidget> ProjectOptionsWidget =
+		SNew(SVerticalBox)
+		.Visibility(this, &SProjectDialog::GetProjectSettingsVisibility)
+		+ SVerticalBox::Slot()
+		.AutoHeight()
+		[	
+			SNew(SBorder)
+			.Padding(FMargin(10.0f, 7.0f))
+			.BorderImage(FAppStyle::Get().GetBrush("Brushes.Header"))
+			[
+				SNew(STextBlock)	
+				.TextStyle(FEditorStyle::Get(), "DialogButtonText")
+				.Text(LOCTEXT("ProjectDefaults", "Project Defaults"))
+			]
+		]
+		+ SVerticalBox::Slot()
+		.Padding(0.0f, 20.0f, 0.0f, 0.0f)
+		[
+			SNew(SScrollBox)
+			+ SScrollBox::Slot()
+			[
+				SAssignNew(ProjectOptionsBox, SVerticalBox)
+			]
+		];
+
+	const TArray<ETemplateSetting>& HiddenSettings = GetSelectedTemplateProperty(&FTemplateItem::HiddenSettings);
+
+	bool bIsBlueprintAvailable = !GetSelectedTemplateProperty(&FTemplateItem::BlueprintProjectFile).IsEmpty();
+	bool bIsCodeAvailable = !GetSelectedTemplateProperty(&FTemplateItem::CodeProjectFile).IsEmpty();
+
+	if (!HiddenSettings.Contains(ETemplateSetting::Languages))
+	{
+		// if neither is available, then this is a blank template, so both are available
+		if (!bIsBlueprintAvailable && !bIsCodeAvailable)
+		{
+			bIsBlueprintAvailable = true;
+			bIsCodeAvailable = true;
+		}
+
+		bShouldGenerateCode = !bIsBlueprintAvailable;
+
+		TSharedRef<SSegmentedControl<int32>> ScriptTypeChooser =
+			SNew(SSegmentedControl<int32>)
+			.UniformPadding(FMargin(25.0f,4.0f))
+			
+			.ToolTipText(LOCTEXT("ProjectDialog_BlueprintOrCppDescription", "Choose whether to create a Blueprint or C++ project.\nNote: You can also add blueprints to a C++ project and C++ to a Blueprint project later."))
+			.Value(this, &SProjectDialog::OnGetBlueprintOrCppIndex)
+			.OnValueChanged(this, &SProjectDialog::OnSetBlueprintOrCppIndex);
+
+			if (bIsBlueprintAvailable)
+			{
+				ScriptTypeChooser->AddSlot(0)
+				.HAlign(HAlign_Center)
+				.Text(LOCTEXT("ProjectDialog_Blueprint", "BLUEPRINT"));
+			}
+
+			if (bIsCodeAvailable)
+			{
+				ScriptTypeChooser->AddSlot(1)
+				.HAlign(HAlign_Center)
+				.Text(LOCTEXT("ProjectDialog_Code", "C++"));
+			}
+
+			ProjectOptionsBox->AddSlot()
+				.AutoHeight()
+				.HAlign(HAlign_Center)
+				[
+					ScriptTypeChooser
+				];
+	}
+	else
+	{
+		bShouldGenerateCode = bIsCodeAvailable;
+	}
+
+	if (!HiddenSettings.Contains(ETemplateSetting::HardwareTarget))
+	{
+		TSharedRef<SWidget> HardwareClassTarget = HardwareTargeting.MakeHardwareClassTargetCombo(
+			FOnHardwareClassChanged::CreateSP(this, &SProjectDialog::SetHardwareClassTarget),
+			TAttribute<EHardwareClass>(this, &SProjectDialog::GetHardwareClassTarget));
+		
+		ProjectOptionsBox->AddSlot()
+		.Padding(0.0f, 16.0f, 0.0f, 8.0f)
+		.AutoHeight()
+		[
+			SNew(SHorizontalBox)
+			.ToolTipText(LOCTEXT("ProjectDialog_HardwareClassTargetDescription", "Choose the closest equivalent target platform. You can change this later in the Target Hardware section of Project Settings."))
+			+ SHorizontalBox::Slot()
+			.Padding(0.0f, 0.0f, 8.0f, 0.0f)
+			.VAlign(VAlign_Center)
+			.HAlign(HAlign_Right)
+			[
+				SNew(STextBlock)
+				.Text(LOCTEXT("TargetPlatform", "Target Platform"))
+			]
+			+ SHorizontalBox::Slot()
+			.VAlign(VAlign_Center)
+			.HAlign(HAlign_Left)
+			[
+				SNew(SBox)
+				.WidthOverride(120.0f)
+				[
+					HardwareClassTarget
+				]
+			]
+		];
+	}
+
+	if (!HiddenSettings.Contains(ETemplateSetting::GraphicsPreset))
+	{
+		TSharedRef<SWidget> GraphicsPreset = HardwareTargeting.MakeGraphicsPresetTargetCombo(
+			FOnGraphicsPresetChanged::CreateSP(this, &SProjectDialog::SetGraphicsPreset),
+			TAttribute<EGraphicsPreset>(this, &SProjectDialog::GetGraphicsPreset));
+
+		ProjectOptionsBox->AddSlot()
+		.AutoHeight()
+		[
+			SNew(SHorizontalBox)
+			.ToolTipText(LOCTEXT("ProjectDialog_GraphicsPresetDescription", "Choose the performance characteristics of your project. You can change this later in the Target Hardware section of Project Settings."))
+			+ SHorizontalBox::Slot()
+			.Padding(0.0f, 0.0f, 8.0f, 0.0f)
+			.VAlign(VAlign_Center)
+			.HAlign(HAlign_Right)
+			[
+				SNew(STextBlock)
+				.Text(LOCTEXT("QualityPreset", "Quality Preset"))
+			]
+			+ SHorizontalBox::Slot()
+			.VAlign(VAlign_Center)
+			.HAlign(HAlign_Left)
+			[
+				SNew(SBox)
+				.WidthOverride(120.0f)
+				[
+					GraphicsPreset
+				]
+			]
+		];
+	}
+
+	if (!HiddenSettings.Contains(ETemplateSetting::StarterContent))
+	{
+		FProjectInformation ProjectInfo = CreateProjectInfo();
+
+		const bool bIsStarterContentAvailable = GameProjectUtils::IsStarterContentAvailableForProject(ProjectInfo);
+		ProjectOptionsBox->AddSlot()
+			.Padding(0.0f, 8.0f, 0.0f, 0.0f)
+			.AutoHeight()
+			[
+				SNew(SHorizontalBox)
+				.ToolTipText(LOCTEXT("CopyStarterContent_ToolTip", "Enable to include an additional content pack containing simple placeable meshes with basic materials and textures.\nYou can also add the Starter Content to your project later using the Content Browser."))
+				.IsEnabled(bIsStarterContentAvailable)
+				+ SHorizontalBox::Slot()
+				.Padding(0.0f, 0.0f, 8.0f, 0.0f)
+				.VAlign(VAlign_Center)
+				.HAlign(HAlign_Right)
+				[
+					SNew(STextBlock)
+					.Text(LOCTEXT("IncludeStarterContent", "Starter Content"))
+				]
+				+ SHorizontalBox::Slot()
+				.VAlign(VAlign_Center)
+				[
+					SNew(SCheckBox)
+					.IsChecked(this, &SProjectDialog::GetCopyStarterContentCheckState)
+					.OnCheckStateChanged(this, &SProjectDialog::OnSetCopyStarterContent)
+				]
+			];
+
+		// Only add the option to add starter content if its there to add!
+		if (!bIsStarterContentAvailable)
+		{
+			bCopyStarterContent = false;
+		}
+/*
+
+		TSharedRef<SWidget> Enum = SNew(SOverlay)
+			+ SOverlay::Slot()
+			[
+				SNew(SDecoratedEnumCombo<int32>, MoveTemp(StarterContentOptions))
+				.SelectedEnum()
+				.OnEnumChanged()
+				.Orientation(Orient_Vertical)
+			]
+			+ SOverlay::Slot()
+			.HAlign(HAlign_Right)
+			.VAlign(VAlign_Top)
+			.Padding(4)
+			[
+				// Warning when enabled for mobile, since the current starter content is bad for mobile
+				SNew(SImage)
+				.Image(FEditorStyle::GetBrush("Icons.Warning"))
+				.ToolTipText(this, &SNewProjectWizard::GetStarterContentWarningTooltip)
+				.Visibility(this, &SNewProjectWizard::GetStarterContentWarningVisibility)
+			];
+*/
+
+		//AddToProjectSettingsGrid(GridPanel, Enum, Description, CurrentSlot);
+	}
+
+#if 0 // @todo: XR settings cannot be shown at the moment as the setting causes issues with binary builds.
+	if (!HiddenSettings.Contains(ETemplateSetting::XR))
+	{
+		TArray<SDecoratedEnumCombo<int32>::FComboOption> VirtualRealityOptions;
+		VirtualRealityOptions.Add(SDecoratedEnumCombo<int32>::FComboOption(
+			0, FSlateIcon(FEditorStyle::GetStyleSetName(), "GameProjectDialog.XRDisabled"),
+			LOCTEXT("XRDisabled", "XR Disabled")));
+
+		VirtualRealityOptions.Add(SDecoratedEnumCombo<int32>::FComboOption(
+			1, 
+			FSlateIcon(FEditorStyle::GetStyleSetName(), "GameProjectDialog.XREnabled"),
+			LOCTEXT("XREnabled", "XR Enabled")));
+
+		TSharedRef<SDecoratedEnumCombo<int32>> Enum = SNew(SDecoratedEnumCombo<int32>, MoveTemp(VirtualRealityOptions))
+			.SelectedEnum(this, &SProjectDialog::OnGetXREnabled)
+			.OnEnumChanged(this, &SProjectDialog::OnSetXREnabled)
+			.Orientation(Orient_Vertical);
+
+		TSharedRef<SRichTextBlock> Description = SNew(SRichTextBlock)
+			.Text(LOCTEXT("ProjectDialog_XREnabledDescription", "Choose if XR should be enabled in the new project."))
+			.AutoWrapText(true)
+			.DecoratorStyleSet(&FEditorStyle::Get());
+	}
+#endif 
+
+	if (!HiddenSettings.Contains(ETemplateSetting::Raytracing))
+	{
+		ProjectOptionsBox->AddSlot()
+			.Padding(0.0f, 8.0f, 0.0f, 0.0f)
+			.AutoHeight()
+			[
+				SNew(SHorizontalBox)
+				.ToolTipText(LOCTEXT("ProjectDialog_RaytracingDescription", "Choose if real-time raytracing should be supported in the new project."))
+				+ SHorizontalBox::Slot()
+				.Padding(0.0f, 0.0f, 8.0f, 0.0f)
+				.VAlign(VAlign_Center)
+				.HAlign(HAlign_Right)
+				[
+					SNew(STextBlock)
+					.Text(LOCTEXT("RayTracing", "Raytracing"))
+				]
+				+ SHorizontalBox::Slot()
+				.VAlign(VAlign_Center)
+				[
+					SNew(SCheckBox)
+					.IsChecked(this, &SProjectDialog::OnGetRaytracingEnabledCheckState)
+					.OnCheckStateChanged(this, &SProjectDialog::OnSetRaytracingEnabled)
+				]
+			];
+	}
+
+	return ProjectOptionsWidget;
+}
+
+TSharedRef<SWidget> SProjectDialog::MakeRecentProjectsTile()
+{
+	RecentProjectsCategory = MakeShared<FTemplateCategory>();
+
+	RecentProjectsCategory->DisplayName = LOCTEXT("RecentProjects", "Recent Projects");
+	RecentProjectsCategory->Description = FText::GetEmpty();
+	RecentProjectsCategory->Key = "RecentProjects";
+	RecentProjectsCategory->IsEnterprise = false;
+
+	static const FName BrushName = *(FPaths::RootDir() / TEXT("Templates/Media") / TEXT("RecentProjects_2x.png"));
+
+	RecentProjectsBrush = MakeUnique<FSlateDynamicImageBrush>(BrushName, FVector2D(300, 100));
+	RecentProjectsBrush->OutlineSettings.CornerRadii = FVector4(4, 4, 4, 4);
+	RecentProjectsBrush->OutlineSettings.RoundingType = ESlateBrushRoundingType::FixedRadius;
+	RecentProjectsBrush->DrawAs = ESlateBrushDrawType::RoundedBox;
+
+	RecentProjectsCategory->Icon = RecentProjectsBrush.Get();
+
+	return
+		SNew(SButton)
+		.ButtonStyle(FAppStyle::Get(), "InvisibleButton")
+		.OnClicked(this, &SProjectDialog::OnRecentProjectsClicked)
+		.ForegroundColor(FLinearColor::White)
+		.ContentPadding(FMargin(4.0f, 0.0f))
+		[
+			SNew(SMajorCategoryTile, RecentProjectsCategory)
+			.IsSelected_Lambda([this]() { return RecentProjectsCategory == CurrentCategory; })
+		];
+}
+
+TSharedRef<SWidget> SProjectDialog::MakeNewProjectPathArea()
+{
+	return 
+		SNew(SVerticalBox)
+		+SVerticalBox::Slot()
+		.Padding(25.0f, 36.0f, 0.0f, 0)
+		[
+			SNew(SFilepath)
+			.OnBrowseForFolder(this, &SProjectDialog::HandlePathBrowseButtonClicked)
+			.FolderPath(this, &SProjectDialog::GetCurrentProjectFilePath)
+			.WarningText(this, &SProjectDialog::GetNameAndLocationValidityErrorText)
+			.Name(this, &SProjectDialog::GetCurrentProjectFileName)
+			.OnFolderChanged(this, &SProjectDialog::OnCurrentProjectFilePathChanged)
+			.OnNameChanged(this, &SProjectDialog::OnCurrentProjectFileNameChanged)
+		]
+		+SVerticalBox::Slot()
+		.Padding(0.0f, 0.0f, 0.0f, 8.0f)
+		.VAlign(VAlign_Bottom)
+		.HAlign(HAlign_Right)
+		.AutoHeight()
+		[
+			SNew(SHorizontalBox)
+			+SHorizontalBox::Slot()
+			.VAlign(VAlign_Center)
+			.AutoWidth()
+			[
+				SNew(SWarningOrErrorBox)
+				.Padding(FMargin(8.0f, 4.0f, 4.0f, 4.0f))
+				.IconSize(FVector2D(16,16))
+				.MessageStyle(EMessageStyle::Error)
+				.Message(this, &SProjectDialog::GetGlobalErrorLabelText)
+				.Visibility(this, &SProjectDialog::GetGlobalErrorVisibility)
+			]
+			+SHorizontalBox::Slot()
+			.VAlign(VAlign_Bottom)
+			[
+				MakeNewProjectDialogButtons()
+			]
+		];
+}
+
+TSharedRef<SWidget> SProjectDialog::MakeOpenProjectPathArea()
+{
+	return
+		SNew(SVerticalBox)
+		+ SVerticalBox::Slot()
+		.Padding(25.0f, 36.0f, 0.0f, 0)
+		[
+			SNew(SFilepath)
+			.IsReadOnly(true)
+			.FolderPath(this, &SProjectDialog::GetCurrentProjectFilePath)
+			.Name(this, &SProjectDialog::GetCurrentProjectFileName)
+		]
+		+SVerticalBox::Slot()
+		.Padding(25.0f, 0.0f, 0.0f, 8.0f)
+		.VAlign(VAlign_Bottom)	
+		.AutoHeight()
+		[
+			SNew(SHorizontalBox)
+			+SHorizontalBox::Slot()
+			[
+				SNew(SCheckBox)
+				.IsChecked(GetDefault<UEditorSettings>()->bLoadTheMostRecentlyLoadedProjectAtStartup ? ECheckBoxState::Checked : ECheckBoxState::Unchecked)
+				.OnCheckStateChanged(ProjectBrowser.ToSharedRef(), &SProjectBrowser::OnAutoloadLastProjectChanged)
+				.Padding(FMargin(4.0f, 0.0f))
+				[
+					SNew(STextBlock)
+					.Text(LOCTEXT("AutoloadOnStartupCheckbox", "Always load last project on startup"))
+					.ColorAndOpacity(FSlateColor::UseForeground())
+				]
+			]
+			+ SHorizontalBox::Slot()
+			.VAlign(VAlign_Bottom)
+			.HAlign(HAlign_Right)
+			[
+				MakeOpenProjectDialogButtons()
+			]
+		];
+}
+
+void SProjectDialog::OnSetCopyStarterContent(ECheckBoxState NewState)
+{
+	bCopyStarterContent = NewState == ECheckBoxState::Checked;
+}
+
+bool SProjectDialog::CanCreateProject() const
+{
+	return bLastGlobalValidityCheckSuccessful && bLastNameAndLocationValidityCheckSuccessful;
+}
+
+FReply SProjectDialog::OnCancel() const
+{
+	TSharedPtr<SWindow> Window = FSlateApplication::Get().FindWidgetWindow(AsShared());
+	Window->RequestDestroyWindow();
+
+	return FReply::Handled();
+}
+
+void SProjectDialog::OnSetRaytracingEnabled(ECheckBoxState NewState)
+{
+	bEnableRaytracing = NewState == ECheckBoxState::Checked;
+}
+
+void SProjectDialog::OnSetBlueprintOrCppIndex(int32 Index)
+{
+	bShouldGenerateCode = Index == 1;
+	UpdateProjectFileValidity();
+}
+
+void SProjectDialog::SetHardwareClassTarget(EHardwareClass InHardwareClass)
+{
+	SelectedHardwareClassTarget = InHardwareClass;
+}
+
+void SProjectDialog::SetGraphicsPreset(EGraphicsPreset InGraphicsPreset)
+{
+	SelectedGraphicsPreset = InGraphicsPreset;
+}
+
+void SProjectDialog::HandleTemplateListViewSelectionChanged(TSharedPtr<FTemplateItem> TemplateItem, ESelectInfo::Type SelectInfo)
+{
+	if (TemplateItem.IsValid())
+	{
+		if (TemplateItem->HiddenSettings.Contains(ETemplateSetting::StarterContent))
+		{
+			bCopyStarterContent = false;
+		}
+	}
+
+	(*ProjectOptionsSlot)
+	[
+		MakeProjectOptionsWidget()
+	];
+
+	UpdateProjectFileValidity();
+}
+
+TSharedPtr<FTemplateItem> SProjectDialog::GetSelectedTemplateItem() const
+{
+	TArray<TSharedPtr<FTemplateItem>> SelectedItems = TemplateListView->GetSelectedItems();
+	if (SelectedItems.Num() > 0)
+	{
+		return SelectedItems[0];
+	}
+
+	return nullptr;
+}
+
+TMap<FName, TArray<TSharedPtr<FTemplateItem>> > SProjectDialog::FindTemplateProjects()
+{
+	// Clear the list out first - or we could end up with duplicates
+	TMap<FName, TArray<TSharedPtr<FTemplateItem>>> Templates;
+
+	// Now discover and all data driven templates
+	TArray<FString> TemplateRootFolders;
+
+	// @todo rocket make template folder locations extensible.
+	TemplateRootFolders.Add(FPaths::RootDir() + TEXT("Templates"));
+
+	// Add the Enterprise templates
+	TemplateRootFolders.Add(FPaths::EnterpriseDir() + TEXT("Templates"));
+
+	// Allow plugins to define templates
+	TArray<TSharedRef<IPlugin>> Plugins = IPluginManager::Get().GetEnabledPlugins();
+	for (const TSharedRef<IPlugin>& Plugin : Plugins)
+	{
+		FString PluginDirectory = Plugin->GetBaseDir();
+		if (!PluginDirectory.IsEmpty())
+		{
+			const FString PluginTemplatesDirectory = FPaths::Combine(*PluginDirectory, TEXT("Templates"));
+
+			if (IFileManager::Get().DirectoryExists(*PluginTemplatesDirectory))
+			{
+				TemplateRootFolders.Add(PluginTemplatesDirectory);
+			}
+		}
+	}
+
+	// Form a list of all folders that could contain template projects
+	TArray<FString> AllTemplateFolders;
+	for (const FString& Root : TemplateRootFolders)
+	{
+		const FString SearchString = Root / TEXT("*");
+		TArray<FString> TemplateFolders;
+		IFileManager::Get().FindFiles(TemplateFolders, *SearchString, /*Files=*/false, /*Directories=*/true);
+
+		for (const FString& Folder : TemplateFolders)
+		{
+			AllTemplateFolders.Add(Root / Folder);
+		}
+	}
+
+	TArray<TSharedPtr<FTemplateItem>> FoundTemplates;
+
+	// Add a template item for every discovered project
+	for (const FString& Root : AllTemplateFolders)
+	{
+		const FString SearchString = Root / TEXT("*.") + FProjectDescriptor::GetExtension();
+		TArray<FString> FoundProjectFiles;
+		IFileManager::Get().FindFiles(FoundProjectFiles, *SearchString, /*Files=*/true, /*Directories=*/false);
+
+		if (FoundProjectFiles.Num() == 0 || !ensure(FoundProjectFiles.Num() == 1))
+		{
+			continue;
+		}
+
+		// Make sure a TemplateDefs.ini file exists
+		UTemplateProjectDefs* TemplateDefs = GameProjectUtils::LoadTemplateDefs(Root);
+		if (TemplateDefs == nullptr)
+		{
+			continue;
+		}
+
+		// we don't have an appropriate referencing UObject to keep these alive with, so we need to keep these template defs alive from GC
+		TemplateDefs->AddToRoot();
+
+		// Ignore any templates whose definition says we cannot use to create a project
+		if (TemplateDefs->bAllowProjectCreation == false)
+		{
+			continue;
+		}
+
+		const FString ProjectFile = Root / FoundProjectFiles[0];
+
+		// If no template category was specified, use the default category
+		TArray<FName> TemplateCategoryNames = TemplateDefs->Categories;
+		if (TemplateCategoryNames.Num() == 0)
+		{
+			TemplateCategoryNames.Add(NewProjectDialogDefs::DefaultCategoryName);
+		}
+
+		// Find a duplicate project, eg. "Foo" and "FooBP"
+		FString TemplateKey = Root;
+		TemplateKey.RemoveFromEnd("BP");
+
+		TSharedPtr<FTemplateItem>* ExistingTemplate = FoundTemplates.FindByPredicate([&TemplateKey](TSharedPtr<FTemplateItem> Item)
+			{
+				return Item->Key == TemplateKey;
+			});
+
+		TSharedPtr<FTemplateItem> Template;
+
+		// Create a new template if none was found
+		if (ExistingTemplate != nullptr)
+		{
+			Template = *ExistingTemplate;
+		}
+		else
+		{
+			Template = MakeShareable(new FTemplateItem());
+		}
+
+		if (TemplateDefs->GeneratesCode(Root))
+		{
+			Template->CodeProjectFile = ProjectFile;
+			Template->CodeTemplateDefs = TemplateDefs;
+		}
+		else
+		{
+			Template->BlueprintProjectFile = ProjectFile;
+			Template->BlueprintTemplateDefs = TemplateDefs;
+		}
+
+		// The rest has already been set by the existing template, so skip it.
+		if (ExistingTemplate != nullptr)
+		{
+			continue;
+		}
+
+		// Did not find an existing template. Create a new one to add to the template list.
+		Template->Key = TemplateKey;
+
+		// @todo: These are all basically just copies of what's in UTemplateProjectDefs, but ignore differences between code and BP 
+		Template->Categories = TemplateCategoryNames;
+		Template->Description = TemplateDefs->GetLocalizedDescription();
+		Template->ClassTypes = TemplateDefs->ClassTypes;
+		Template->AssetTypes = TemplateDefs->AssetTypes;
+		Template->HiddenSettings = TemplateDefs->HiddenSettings;
+		Template->bIsEnterprise = TemplateDefs->bIsEnterprise;
+		Template->bIsBlankTemplate = TemplateDefs->bIsBlank;
+		Template->bThumbnailAsIcon = TemplateDefs->bThumbnailAsIcon;
+
+		Template->Name = TemplateDefs->GetDisplayNameText();
+		if (Template->Name.IsEmpty())
+		{
+			Template->Name = FText::FromString(TemplateKey);
+		}
+
+		const FString ThumbnailPNGFile = (Root + TEXT("/Media/") + FoundProjectFiles[0]).Replace(TEXT(".uproject"), TEXT(".png"));
+		if (FPlatformFileManager::Get().GetPlatformFile().FileExists(*ThumbnailPNGFile))
+		{
+			const FName BrushName = FName(*ThumbnailPNGFile);
+			Template->Thumbnail = MakeShareable(new FSlateDynamicImageBrush(BrushName, FVector2D(128, 128)));
+		}
+
+		TSharedPtr<FSlateDynamicImageBrush> PreviewBrush;
+		const FString PreviewPNGFile = (Root + TEXT("/Media/") + FoundProjectFiles[0]).Replace(TEXT(".uproject"), TEXT("_Preview.png"));
+		if (FPlatformFileManager::Get().GetPlatformFile().FileExists(*PreviewPNGFile))
+		{
+			const FName BrushName = FName(*PreviewPNGFile);
+			Template->PreviewImage = MakeShareable(new FSlateDynamicImageBrush(BrushName, FVector2D(512, 256)));
+		}
+
+		Template->SortKey = TemplateDefs->SortKey;
+		if (Template->SortKey.IsEmpty())
+		{
+			Template->SortKey = TemplateKey;
+		}
+
+		FoundTemplates.Add(Template);
+	}
+
+	for (const TSharedPtr<FTemplateItem>& Template : FoundTemplates)
+	{
+		for (const FName& Category : Template->Categories)
+		{
+			Templates.FindOrAdd(Category).Add(Template);
+		}
+	}
+
+	TArray<TSharedPtr<FTemplateCategory>> AllTemplateCategories = GetAllTemplateCategories();
+
+	// Validate that all our templates have a category defined
+	TArray<FName> CategoryKeys;
+	Templates.GetKeys(CategoryKeys);
+	for (const FName& CategoryKey : CategoryKeys)
+	{
+		bool bCategoryExists = AllTemplateCategories.ContainsByPredicate([&CategoryKey](const TSharedPtr<FTemplateCategory>& Category)
+			{
+				return Category->Key == CategoryKey;
+			});
+
+		if (!bCategoryExists)
+		{
+			UE_LOG(LogGameProjectGeneration, Warning, TEXT("Failed to find category definition named '%s', it is not defined in any TemplateCategories.ini."), *CategoryKey.ToString());
+		}
+	}
+
+	// Add blank template to empty categories
+	{
+		TSharedPtr<FTemplateItem> BlankTemplate = MakeShareable(new FTemplateItem());
+		BlankTemplate->Name = LOCTEXT("BlankProjectName", "Blank");
+		BlankTemplate->Description = LOCTEXT("BlankProjectDescription", "A clean empty project with no code and default settings.");
+		BlankTemplate->Key = TEXT("Blank");
+		BlankTemplate->SortKey = TEXT("_1");
+		BlankTemplate->Thumbnail = MakeShareable(new FSlateBrush(*FEditorStyle::GetBrush("GameProjectDialog.BlankProjectThumbnail")));
+		BlankTemplate->PreviewImage = MakeShareable(new FSlateBrush(*FEditorStyle::GetBrush("GameProjectDialog.BlankProjectPreview")));
+		BlankTemplate->BlueprintProjectFile = TEXT("");
+		BlankTemplate->CodeProjectFile = TEXT("");
+		BlankTemplate->bIsEnterprise = false;
+		BlankTemplate->bIsBlankTemplate = true;
+
+		for (const TSharedPtr<FTemplateCategory>& Category : AllTemplateCategories)
+		{
+			const TArray<TSharedPtr<FTemplateItem>>* CategoryEntry = Templates.Find(Category->Key);
+			if (CategoryEntry == nullptr)
+			{
+				Templates.Add(Category->Key).Add(BlankTemplate);
+			}
+		}
+	}
+
+	return Templates;
+}
+
+
+void SProjectDialog::SetDefaultProjectLocation()
+{
+	FString DefaultProjectFilePath;
+
+	// First, try and use the first previously used path that still exists
+	for (const FString& CreatedProjectPath : GetDefault<UEditorSettings>()->CreatedProjectPaths)
+	{
+		if (IFileManager::Get().DirectoryExists(*CreatedProjectPath))
+		{
+			DefaultProjectFilePath = CreatedProjectPath;
+			break;
+		}
+	}
+
+	if (DefaultProjectFilePath.IsEmpty())
+	{
+		// No previously used path, decide a default path.
+		DefaultProjectFilePath = FDesktopPlatformModule::Get()->GetDefaultProjectCreationPath();
+		IFileManager::Get().MakeDirectory(*DefaultProjectFilePath, true);
+	}
+
+	if (DefaultProjectFilePath.EndsWith(TEXT("/")))
+	{
+		DefaultProjectFilePath.LeftChopInline(1);
+	}
+
+	FPaths::NormalizeFilename(DefaultProjectFilePath);
+	FPaths::MakePlatformFilename(DefaultProjectFilePath);
+	const FString GenericProjectName = LOCTEXT("DefaultProjectName", "MyProject").ToString();
+	FString ProjectName = GenericProjectName;
+
+	// Check to make sure the project file doesn't already exist
+	FText FailReason;
+	if (!GameProjectUtils::IsValidProjectFileForCreation(DefaultProjectFilePath / ProjectName / ProjectName + TEXT(".") + FProjectDescriptor::GetExtension(), FailReason))
+	{
+		// If it exists, find an appropriate numerical suffix
+		const int MaxSuffix = 1000;
+		int32 Suffix;
+		for (Suffix = 2; Suffix < MaxSuffix; ++Suffix)
+		{
+			ProjectName = GenericProjectName + FString::Printf(TEXT("%d"), Suffix);
+			if (GameProjectUtils::IsValidProjectFileForCreation(DefaultProjectFilePath / ProjectName / ProjectName + TEXT(".") + FProjectDescriptor::GetExtension(), FailReason))
+			{
+				// Found a name that is not taken. Break out.
+				break;
+			}
+		}
+
+		if (Suffix >= MaxSuffix)
+		{
+			UE_LOG(LogGameProjectGeneration, Warning, TEXT("Failed to find a suffix for the default project name"));
+			ProjectName = TEXT("");
+		}
+	}
+
+	if (!DefaultProjectFilePath.IsEmpty())
+	{
+		CurrentProjectFileName = ProjectName;
+		CurrentProjectFilePath = DefaultProjectFilePath;
+		FPaths::MakePlatformFilename(CurrentProjectFilePath);
+		LastBrowsePath = CurrentProjectFilePath;
+	}
+}
+
+void SProjectDialog::SetCurrentMajorCategory(FName Category)
+{
+	FilteredTemplateList = Templates.FindRef(Category);
+
+	// Sort the template folders
+	FilteredTemplateList.Sort([](const TSharedPtr<FTemplateItem>& A, const TSharedPtr<FTemplateItem>& B) {
+		return A->SortKey < B->SortKey;
+		});
+
+	if (FilteredTemplateList.Num() > 0)
+	{
+		TemplateListView->SetSelection(FilteredTemplateList[0]);
+	}
+	TemplateListView->RequestListRefresh();
+
+	ActiveCategory = Category;
+}
+
+FReply SProjectDialog::OnRecentProjectsClicked()
+{
+	CurrentCategory = RecentProjectsCategory;
+
+	ActiveCategory = NAME_None;
+
+	MajorCategoryList->ClearSelection();
+
+	TemplateAndRecentProjectsSwitcher->SetActiveWidgetIndex(1);
+	PathAreaSwitcher->SetActiveWidgetIndex(1);
+
+	return FReply::Handled();
+}
+
+FProjectInformation SProjectDialog::CreateProjectInfo() const
+{
+	TSharedPtr<FTemplateItem> SelectedTemplate = GetSelectedTemplateItem();
+	if (!SelectedTemplate.IsValid())
+	{
+		return FProjectInformation();
+	}
+
+	FProjectInformation ProjectInfo;
+	ProjectInfo.bShouldGenerateCode = bShouldGenerateCode;
+	ProjectInfo.bCopyStarterContent = bCopyStarterContent;
+	ProjectInfo.TemplateFile = bShouldGenerateCode ? SelectedTemplate->CodeProjectFile : SelectedTemplate->BlueprintProjectFile;
+	ProjectInfo.TemplateCategory = ActiveCategory;
+	ProjectInfo.bIsEnterpriseProject = SelectedTemplate->bIsEnterprise;
+	ProjectInfo.bIsBlankTemplate = SelectedTemplate->bIsBlankTemplate;
+	ProjectInfo.bForceExtendedLuminanceRange = SelectedTemplate->bIsBlankTemplate;
+
+	if (bCopyStarterContent)
+	{
+		if (bShouldGenerateCode)
+		{
+			if (SelectedTemplate->CodeTemplateDefs != nullptr)
+			{
+				ProjectInfo.StarterContent = SelectedTemplate->CodeTemplateDefs->StarterContent;
+			}
+		}
+		else
+		{
+			if (SelectedTemplate->BlueprintTemplateDefs != nullptr)
+			{
+				ProjectInfo.StarterContent = SelectedTemplate->BlueprintTemplateDefs->StarterContent;
+			}
+		}
+	}
+
+	const TArray<ETemplateSetting>& HiddenSettings = SelectedTemplate->HiddenSettings;
+
+	if (!HiddenSettings.Contains(ETemplateSetting::All))
+	{
+		if (!HiddenSettings.Contains(ETemplateSetting::HardwareTarget))
+		{
+			ProjectInfo.TargetedHardware = SelectedHardwareClassTarget;
+		}
+
+		if (!HiddenSettings.Contains(ETemplateSetting::GraphicsPreset))
+		{
+			ProjectInfo.DefaultGraphicsPerformance = SelectedGraphicsPreset;
+		}
+
+		if (!HiddenSettings.Contains(ETemplateSetting::XR))
+		{
+			ProjectInfo.bEnableXR = bEnableXR;
+		}
+
+		if (!HiddenSettings.Contains(ETemplateSetting::Raytracing))
+		{
+			ProjectInfo.bEnableRaytracing = bEnableRaytracing;
+		}
+	}
+
+	return MoveTemp(ProjectInfo);
+}
+
+bool SProjectDialog::CreateProject(const FString& ProjectFile)
+{
+	// Get the selected template
+	TSharedPtr<FTemplateItem> SelectedTemplate = GetSelectedTemplateItem();
+
+	if (!ensure(SelectedTemplate.IsValid()))
+	{
+		// A template must be selected.
+		return false;
+	}
+
+	FText FailReason, FailLog;
+
+	FProjectInformation ProjectInfo = CreateProjectInfo();
+	ProjectInfo.ProjectFilename = ProjectFile;
+
+	if (!GameProjectUtils::CreateProject(ProjectInfo, FailReason, FailLog))
+	{
+		SOutputLogDialog::Open(LOCTEXT("CreateProject", "Create Project"), FailReason, FailLog, FText::GetEmpty());
+		return false;
+	}
+
+	// Successfully created the project. Update the last created location string.
+	FString CreatedProjectPath = FPaths::GetPath(FPaths::GetPath(ProjectFile));
+
+	// If the original path was the drives root (ie: C:/) the double path call strips the last /
+	if (CreatedProjectPath.EndsWith(":"))
+	{
+		CreatedProjectPath.AppendChar('/');
+	}
+
+	UEditorSettings* Settings = GetMutableDefault<UEditorSettings>();
+	Settings->CreatedProjectPaths.Remove(CreatedProjectPath);
+	Settings->CreatedProjectPaths.Insert(CreatedProjectPath, 0);
+	Settings->bCopyStarterContentPreference = bCopyStarterContent;
+	Settings->PostEditChange();
+
+	return true;
+}
+
+void SProjectDialog::CreateAndOpenProject()
+{
+	if (!CanCreateProject())
+	{
+		return;
+	}
+
+	FString ProjectFile = GetProjectFilenameWithPath();
+	if (!CreateProject(ProjectFile))
+	{
+		return;
+	}
+
+	if (bShouldGenerateCode)
+	{
+		// If the engine is installed it is already compiled, so we can try to build and open a new project immediately. Non-installed situations might require building
+		// the engine (especially the case when binaries came from P4), so we only open the IDE for that.
+		if (FApp::IsEngineInstalled())
+		{
+			if (GameProjectUtils::BuildCodeProject(ProjectFile))
+			{
+				OpenCodeIDE(ProjectFile);
+				OpenProject(ProjectFile);
+			}
+			else
+			{
+				// User will have already been prompted to open the IDE
+			}
+		}
+		else
+		{
+			OpenCodeIDE(ProjectFile);
+		}
+	}
+	else
+	{
+		OpenProject(ProjectFile);
+	}
+}
+
+bool SProjectDialog::OpenProject(const FString& ProjectFile)
+{
+	FText FailReason;
+	if (GameProjectUtils::OpenProject(ProjectFile, FailReason))
+	{
+		// Successfully opened the project, the editor is closing.
+		// Close this window in case something prevents the editor from closing (save dialog, quit confirmation, etc)
+		CloseWindowIfAppropriate();
+		return true;
+	}
+
+	DisplayError(FailReason);
+	return false;
+}
+
+void SProjectDialog::CloseWindowIfAppropriate(bool ForceClose)
+{
+	if (ForceClose || FApp::HasProjectName())
+	{
+		TSharedPtr<SWindow> ContainingWindow = FSlateApplication::Get().FindWidgetWindow(AsShared());
+
+		if (ContainingWindow.IsValid())
+		{
+			ContainingWindow->RequestDestroyWindow();
+		}
+	}
+}
+
+void SProjectDialog::DisplayError(const FText& ErrorText)
+{
+	FString ErrorString = ErrorText.ToString();
+	UE_LOG(LogGameProjectGeneration, Log, TEXT("%s"), *ErrorString);
+	if(ErrorString.Contains("\n"))
+	{
+		FMessageDialog::Open(EAppMsgType::Ok, ErrorText);
+	}
+	else
+	{
+		PersistentGlobalErrorLabelText = ErrorText;
+	}
+}
+
+bool SProjectDialog::OpenCodeIDE(const FString& ProjectFile)
+{
+	FText FailReason;
+
+	if (GameProjectUtils::OpenCodeIDE(ProjectFile, FailReason))
+	{
+		// Successfully opened code editing IDE, the editor is closing
+		// Close this window in case something prevents the editor from closing (save dialog, quit confirmation, etc)
+		CloseWindowIfAppropriate(true);
+		return true;
+	}
+
+	DisplayError(FailReason);
+	return false;
+}
+
+
+TArray<TSharedPtr<FTemplateCategory>> SProjectDialog::GetAllTemplateCategories()
+{
+	TArray<TSharedPtr<FTemplateCategory>> AllTemplateCategories;
+	FGameProjectGenerationModule::Get().GetAllTemplateCategories(AllTemplateCategories);
+
+	if (AllTemplateCategories.Num() == 0)
+	{
+		TSharedPtr<FTemplateCategory> DefaultCategory = MakeShareable(new FTemplateCategory());
+		static const FName DefaultCategoryKey("Default");
+		DefaultCategory->Key = DefaultCategoryKey;
+		DefaultCategory->DisplayName = LOCTEXT("ProjectDialog_DefaultCategoryName", "Blank Project");
+		DefaultCategory->Description = LOCTEXT("ProjectDialog_DefaultCategoryDescription", "Create a new blank Unreal project.");
+		DefaultCategory->Icon = FEditorStyle::GetBrush("GameProjectDialog.DefaultGameThumbnail");
+
+		AllTemplateCategories.Add(DefaultCategory);
+	}
+
+	return AllTemplateCategories;
+}
+
+FText SProjectDialog::GetGlobalErrorLabelText() const
+{
+	if (!PersistentGlobalErrorLabelText.IsEmpty())
+	{
+		return PersistentGlobalErrorLabelText;
+	}
+
+	if (!bLastGlobalValidityCheckSuccessful)
+	{
+		return LastGlobalValidityErrorText;
+	}
+
+	return FText::GetEmpty();
+}
+
+FText SProjectDialog::GetNameAndLocationValidityErrorText() const
+{
+	if (GetGlobalErrorLabelText().IsEmpty())
+	{
+		return bLastNameAndLocationValidityCheckSuccessful == false ? LastNameAndLocationValidityErrorText : FText::GetEmpty();
+	}
+
+	return FText::GetEmpty();
+}
+
+EVisibility SProjectDialog::GetCreateButtonVisibility() const
+{
+	return IsCompilerRequired() && !FSourceCodeNavigation::IsCompilerAvailable() ? EVisibility::Collapsed : EVisibility::Visible;
+}
+
+EVisibility SProjectDialog::GetSuggestedIDEButtonVisibility() const
+{
+	return IsCompilerRequired() && !FSourceCodeNavigation::IsCompilerAvailable() ? EVisibility::Visible : EVisibility::Collapsed;
+}
+
+const FSlateBrush* SProjectDialog::GetSelectedTemplatePreviewImage() const
+{
+	TSharedPtr<FSlateBrush> PreviewImage = GetSelectedTemplateProperty(&FTemplateItem::PreviewImage);
+	return PreviewImage.IsValid() ? PreviewImage.Get() : nullptr;
+}
+
+
+FText SProjectDialog::GetCurrentProjectFilePath() const
+{
+	return PathAreaSwitcher && PathAreaSwitcher->GetActiveWidgetIndex() == 1 ? FText::FromString(ProjectBrowser->GetSelectedProjectFile()) : FText::FromString(CurrentProjectFilePath);
+}
+
+void SProjectDialog::OnCurrentProjectFilePathChanged(const FText& InValue)
+{
+	CurrentProjectFilePath = InValue.ToString();
+	FPaths::MakePlatformFilename(CurrentProjectFilePath);
+	UpdateProjectFileValidity();
+}
+
+void SProjectDialog::OnCurrentProjectFileNameChanged(const FText& InValue)
+{
+	CurrentProjectFileName = InValue.ToString();
+	UpdateProjectFileValidity();
+}
+
+FText SProjectDialog::GetCurrentProjectFileName() const
+{
+	return PathAreaSwitcher && PathAreaSwitcher->GetActiveWidgetIndex() == 1 ? ProjectBrowser->GetSelectedProjectName() : FText::FromString(CurrentProjectFileName);
+}
+
+FReply SProjectDialog::HandlePathBrowseButtonClicked()
+{
+	IDesktopPlatform* DesktopPlatform = FDesktopPlatformModule::Get();
+	if (DesktopPlatform)
+	{
+		FString FolderName;
+		const FString Title = LOCTEXT("NewProjectBrowseTitle", "Choose a project location").ToString();
+		const bool bFolderSelected = DesktopPlatform->OpenDirectoryDialog(
+			FSlateApplication::Get().FindBestParentWindowHandleForDialogs(AsShared()),
+			Title,
+			LastBrowsePath,
+			FolderName
+		);
+
+		if (bFolderSelected)
+		{
+			if (!FolderName.EndsWith(TEXT("/")))
+			{
+				FolderName += TEXT("/");
+			}
+
+			FPaths::MakePlatformFilename(FolderName);
+			LastBrowsePath = FolderName;
+			CurrentProjectFilePath = FolderName;
+		}
+	}
+
+	return FReply::Handled();
+}
+
+bool SProjectDialog::IsCompilerRequired() const
+{
+	TSharedPtr<FTemplateItem> SelectedTemplate = GetSelectedTemplateItem();
+
+	if (SelectedTemplate.IsValid())
+	{
+		return bShouldGenerateCode && !SelectedTemplate->CodeProjectFile.IsEmpty();
+	}
+
+	return false;
+}
+
+EVisibility SProjectDialog::GetProjectSettingsVisibility() const
+{
+	const TArray<ETemplateSetting>& HiddenSettings = GetSelectedTemplateProperty(&FTemplateItem::HiddenSettings);
+	return HiddenSettings.Contains(ETemplateSetting::All) ? EVisibility::Collapsed : EVisibility::Visible;
+}
+
+EVisibility SProjectDialog::GetSelectedTemplateClassVisibility() const
+{
+	return GetSelectedTemplateProperty(&FTemplateItem::ClassTypes).IsEmpty() == false ? EVisibility::Visible : EVisibility::Collapsed;
+}
+
+FText SProjectDialog::GetSelectedTemplateAssetTypes() const
+{
+	return FText::FromString(GetSelectedTemplateProperty(&FTemplateItem::AssetTypes));
+}
+
+FText SProjectDialog::GetSelectedTemplateClassTypes() const
+{
+	return FText::FromString(GetSelectedTemplateProperty(&FTemplateItem::ClassTypes));
+}
+
+EVisibility SProjectDialog::GetSelectedTemplateAssetVisibility() const
+{
+	return GetSelectedTemplateProperty(&FTemplateItem::AssetTypes).IsEmpty() == false ? EVisibility::Visible : EVisibility::Collapsed;
+}
+
+FString SProjectDialog::GetProjectFilenameWithPath() const
+{
+	if (CurrentProjectFilePath.IsEmpty())
+	{
+		// Don't even try to assemble the path or else it may be relative to the binaries folder!
+		return TEXT("");
+	}
+	else
+	{
+		const FString ProjectName = CurrentProjectFileName;
+		const FString ProjectPath = IFileManager::Get().ConvertToAbsolutePathForExternalAppForWrite(*CurrentProjectFilePath);
+		const FString Filename = ProjectName + TEXT(".") + FProjectDescriptor::GetExtension();
+		FString ProjectFilename = FPaths::Combine(*ProjectPath, *ProjectName, *Filename);
+		FPaths::MakePlatformFilename(ProjectFilename);
+		return ProjectFilename;
+	}
+}
+
+void SProjectDialog::UpdateProjectFileValidity()
+{
+	// Global validity
+	{
+		bLastGlobalValidityCheckSuccessful = true;
+
+		TSharedPtr<FTemplateItem> SelectedTemplate = GetSelectedTemplateItem();
+		if (!SelectedTemplate.IsValid())
+		{
+			bLastGlobalValidityCheckSuccessful = false;
+			LastGlobalValidityErrorText = LOCTEXT("NoTemplateSelected", "No Template Selected");
+		}
+		else
+		{
+			if (IsCompilerRequired())
+			{
+				if (!FSourceCodeNavigation::IsCompilerAvailable())
+				{
+					bLastGlobalValidityCheckSuccessful = false;
+					LastGlobalValidityErrorText = FText::Format(LOCTEXT("NoCompilerFound", "No compiler was found. In order to use a C++ template, you must first install {0}."), FSourceCodeNavigation::GetSuggestedSourceCodeIDE());
+				}
+				else if (!FDesktopPlatformModule::Get()->IsUnrealBuildToolAvailable())
+				{
+					bLastGlobalValidityCheckSuccessful = false;
+					LastGlobalValidityErrorText = LOCTEXT("UBTNotFound", "Engine source code was not found. In order to use a C++ template, you must have engine source code in Engine/Source.");
+				}
+			}
+		}
+	}
+
+	// Name and Location Validity
+	{
+		bLastNameAndLocationValidityCheckSuccessful = true;
+
+		if (!FPlatformMisc::IsValidAbsolutePathFormat(CurrentProjectFilePath))
+		{
+			bLastNameAndLocationValidityCheckSuccessful = false;
+			LastNameAndLocationValidityErrorText = LOCTEXT("InvalidFolderPath", "The folder path is invalid");
+		}
+		else
+		{
+			FText FailReason;
+			if (!GameProjectUtils::IsValidProjectFileForCreation(GetProjectFilenameWithPath(), FailReason))
+			{
+				bLastNameAndLocationValidityCheckSuccessful = false;
+				LastNameAndLocationValidityErrorText = FailReason;
+			}
+		}
+
+		if (CurrentProjectFileName.Contains(TEXT("/")) || CurrentProjectFileName.Contains(TEXT("\\")))
+		{
+			bLastNameAndLocationValidityCheckSuccessful = false;
+			LastNameAndLocationValidityErrorText = LOCTEXT("SlashOrBackslashInProjectName", "The project name may not contain a slash or backslash");
+		}
+		else
+		{
+			FText FailReason;
+			if (!GameProjectUtils::IsValidProjectFileForCreation(GetProjectFilenameWithPath(), FailReason))
+			{
+				bLastNameAndLocationValidityCheckSuccessful = false;
+				LastNameAndLocationValidityErrorText = FailReason;
+			}
+		}
+	}
+}
+
+void SProjectDialog::OnMajorTemplateCategorySelectionChanged(TSharedPtr<FTemplateCategory> Item, ESelectInfo::Type SelectType)
+{
+	if(Item.IsValid())
+	{
+		CurrentCategory = Item;
+
+		SetCurrentMajorCategory(Item->Key);
+
+		TemplateAndRecentProjectsSwitcher->SetActiveWidgetIndex(0);
+		PathAreaSwitcher->SetActiveWidgetIndex(0);
+	}
+}
+
+TSharedRef<ITableRow> SProjectDialog::ConstructMajorCategoryTableRow(TSharedPtr<FTemplateCategory> Item, const TSharedRef<STableViewBase>& TableView)
+{
+	TWeakPtr<FTemplateCategory> CurrentItem = Item;
+
+	TSharedRef<STableRow<TSharedPtr<FTemplateCategory>>> Row =
+		SNew(STableRow<TSharedPtr<FTemplateCategory>>, TableView)
+		.Style(FAppStyle::Get(), "ProjectBrowser.TableRow")
+		.ShowSelection(false)
+		.Padding(2);
+	TWeakPtr<STableRow<TSharedPtr<FTemplateCategory>>> RowWeakPtr = Row;
+
+	Row->SetContent(
+		SNew(SMajorCategoryTile, Item)
+		.IsSelected_Lambda([CurrentItem, this]() { return CurrentItem == CurrentCategory; })
+	);
+
+	return Row;
+}
+
+#undef LOCTEXT_NAMESPACE
diff --git a/Engine/Source/Editor/GameProjectGeneration/Private/SProjectDialog.h b/Engine/Source/Editor/GameProjectGeneration/Private/SProjectDialog.h
new file mode 100644
--- /dev/null
+++ b/Engine/Source/Editor/GameProjectGeneration/Private/SProjectDialog.h
@@ -0,0 +1,236 @@
+// Copyright Epic Games, Inc. All Rights Reserved.
+
+#pragma once
+
+#include "CoreMinimal.h"
+#include "SlateFwd.h"
+#include "Widgets/DeclarativeSyntaxSupport.h"
+#include "Widgets/SCompoundWidget.h"
+#include "Widgets/Views/STileView.h"
+#include "GameProjectUtils.h"
+
+struct FTemplateCategory;
+class SProjectBrowser;
+class SNewProjectWizard;
+struct FTemplateItem;
+
+enum class EHardwareClass : uint8;
+enum class EGraphicsPreset : uint8;
+
+enum class EProjectDialogModeMode : uint8
+{
+	/** Shows new project templates and existing projects */
+	Hybrid,
+	/** Shows existing projects */
+	OpenProject,
+	/** Shows new project templates */
+	NewProject,
+};
+
+class SProjectDialog : public SCompoundWidget
+{
+	SLATE_BEGIN_ARGS(SProjectDialog) {}
+	SLATE_END_ARGS()
+
+	/** Constructs this widget with InArgs */
+	void Construct(const FArguments& InArgs, EProjectDialogModeMode Mode);
+
+	~SProjectDialog();
+
+	/** @returns All templates found on disk */
+	static TMap<FName, TArray<TSharedPtr<FTemplateItem>>> FindTemplateProjects();
+
+private:
+	/** Accessor for the currently selected template item */
+	TSharedPtr<FTemplateItem> GetSelectedTemplateItem() const;
+
+	/** Helper function to allow direct lookup of the selected item's properties on a delegate */
+	/** TRet should be defaulted but VS2012 doesn't allow default template arguments on non-class templates */
+	template<typename T>
+	T GetSelectedTemplateProperty(T FTemplateItem::* Prop) const
+	{
+		TSharedPtr<FTemplateItem> SelectedItem = GetSelectedTemplateItem();
+		if (SelectedItem.IsValid())
+		{
+			return (*SelectedItem).*Prop;
+		}
+
+		return T();
+	}
+
+	void OnMajorTemplateCategorySelectionChanged(TSharedPtr<FTemplateCategory> Item, ESelectInfo::Type SelectType);
+	TSharedRef<ITableRow> ConstructMajorCategoryTableRow(TSharedPtr<FTemplateCategory> Item, const TSharedRef<STableViewBase>& TableView);
+
+	void PopulateTemplateCategories();
+
+	/** Handler for when the project path string was changed */
+	void OnCurrentProjectFilePathChanged(const FText& InValue);
+
+	/** Handler for when the project name string was changed */
+	void OnCurrentProjectFileNameChanged(const FText& InValue);
+
+	FReply HandlePathBrowseButtonClicked();
+
+	/** Checks the current project path an name for validity and updates cached values accordingly */
+	void UpdateProjectFileValidity();
+
+	/** Returns true if we have a code template selected */
+	bool IsCompilerRequired() const;
+
+	EVisibility GetProjectSettingsVisibility() const;
+
+	/** Get a visibility of the class types display. If the string is empty this return Collapsed otherwise it will return Visible */
+	EVisibility GetSelectedTemplateClassVisibility() const;
+
+	/** Get a visibility of the asset types display. If the string is empty this return Collapsed otherwise it will return Visible */
+	EVisibility GetSelectedTemplateAssetVisibility() const;
+
+	EVisibility GetGlobalErrorVisibility() const { return GetGlobalErrorLabelText().IsEmpty() ? EVisibility::Hidden : EVisibility::HitTestInvisible; }
+
+	EVisibility GetCreateButtonVisibility() const;
+
+	EVisibility GetSuggestedIDEButtonVisibility() const;
+
+	/** Accessor for the project name text */
+	FText GetCurrentProjectFileName() const;
+
+	/** Accessor for the project path text */
+	FText GetCurrentProjectFilePath() const;
+
+	/** Get a visibility of the asset types display. If the string is empty this return Collapsed otherwise it will return Visible */
+	FText GetSelectedTemplateAssetTypes() const;
+
+	/** Get a string that details the class types referenced in the selected template */
+	FText GetSelectedTemplateClassTypes() const;
+
+	FText GetGlobalErrorLabelText() const;
+
+	FText GetNameAndLocationValidityErrorText() const;
+
+	/** Gets the assembled project filename with path */
+	FString GetProjectFilenameWithPath() const;
+
+	TSharedRef<SWidget> MakeNewProjectDialogButtons();
+
+	TSharedRef<SWidget> MakeOpenProjectDialogButtons();
+
+	TSharedRef<SWidget> MakeTemplateProjectView();
+
+	TSharedRef<SWidget> MakeRecentProjectsView();
+
+	TSharedRef<SWidget> MakeHybridView(EProjectDialogModeMode Mode);
+
+	TSharedRef<SWidget> MakeProjectOptionsWidget();
+
+	TSharedRef<SWidget> MakeRecentProjectsTile();
+
+	TSharedRef<SWidget> MakeNewProjectPathArea();
+
+	TSharedRef<SWidget> MakeOpenProjectPathArea();
+
+	/** Get whether we are copying starter content or not */
+	ECheckBoxState GetCopyStarterContentCheckState() const { return bCopyStarterContent ? ECheckBoxState::Checked : ECheckBoxState::Unchecked; };
+
+	/** Handler for when the "copy starter content" checkbox changes state */
+	void OnSetCopyStarterContent(ECheckBoxState NewState);
+
+	/** Get the images for the selected template preview and category */
+	const FSlateBrush* GetSelectedTemplatePreviewImage() const;
+
+	bool CanCreateProject() const;
+
+	FReply OnCancel() const;
+
+	ECheckBoxState OnGetRaytracingEnabledCheckState() const { return bEnableRaytracing ? ECheckBoxState::Checked : ECheckBoxState::Unchecked; }
+	void OnSetRaytracingEnabled(ECheckBoxState NewState);
+
+	int32 OnGetBlueprintOrCppIndex() const { return bShouldGenerateCode ? 1 : 0; }
+	void OnSetBlueprintOrCppIndex(int32 Index);
+
+	void SetHardwareClassTarget(EHardwareClass InHardwareClass);
+	EHardwareClass GetHardwareClassTarget() const { return SelectedHardwareClassTarget; }
+
+	void SetGraphicsPreset(EGraphicsPreset InGraphicsPreset);
+	EGraphicsPreset GetGraphicsPreset() const { return SelectedGraphicsPreset; }
+
+	/** Handler for when the selection changes in the template list */
+	void HandleTemplateListViewSelectionChanged(TSharedPtr<FTemplateItem> TemplateItem, ESelectInfo::Type SelectInfo);
+
+	void SetDefaultProjectLocation();
+
+	void SetCurrentMajorCategory(FName Category);
+
+	FReply OnRecentProjectsClicked();
+
+	FProjectInformation CreateProjectInfo() const;
+	bool CreateProject(const FString& ProjectFile);
+	void CreateAndOpenProject();
+	bool OpenProject(const FString& ProjectFile);
+	bool OpenCodeIDE(const FString& ProjectFile);
+	void CloseWindowIfAppropriate(bool ForceClose = false);
+	void DisplayError(const FText& ErrorText);
+
+
+	static TArray<TSharedPtr<FTemplateCategory>> GetAllTemplateCategories();
+
+private:
+	TArray<TSharedPtr<FTemplateCategory>> TemplateCategories;
+	TSharedPtr<FTemplateCategory> CurrentCategory;
+	TSharedPtr<STileView<TSharedPtr<FTemplateCategory>>> MajorCategoryList;
+
+	TSharedPtr<FTemplateCategory> RecentProjectsCategory;
+
+	/** A map of category name to array of templates available for that category */
+	TMap<FName, TArray<TSharedPtr<FTemplateItem>> > Templates;
+
+	/** The filtered array of templates we are currently showing */
+	TArray<TSharedPtr<FTemplateItem> > FilteredTemplateList;
+
+	/** The slate widget representing the list of templates */
+	TSharedPtr<STileView<TSharedPtr<FTemplateItem>>> TemplateListView;
+
+	TSharedPtr<SWidgetSwitcher> TemplateAndRecentProjectsSwitcher;
+	TSharedPtr<SWidgetSwitcher> PathAreaSwitcher;
+
+	TSharedPtr<SProjectBrowser> ProjectBrowser; 
+
+	FString LastBrowsePath;
+	FString CurrentProjectFileName;
+	FString CurrentProjectFilePath;
+
+	FText PersistentGlobalErrorLabelText;
+
+	/** The global error text from the last validity check */
+	FText LastGlobalValidityErrorText;
+
+	FText LastNameAndLocationValidityErrorText;
+
+	/** Name of the currently selected template category */
+	FName ActiveCategory;
+
+	EHardwareClass SelectedHardwareClassTarget;
+
+	EGraphicsPreset SelectedGraphicsPreset;
+
+	TUniquePtr<FSlateBrush> RecentProjectsBrush;
+
+	SVerticalBox::FSlot* ProjectOptionsSlot;
+	/** True if user has selected to copy starter content. */
+	bool bCopyStarterContent;
+
+	/** Whether or not to enable XR in the created project. */
+	bool bEnableXR;
+
+	/** Whether or not to enable Raytracing in the created project. */
+	bool bEnableRaytracing;
+
+	/** Whether or not we should use the blueprint or C++ version of this template. */
+	bool bShouldGenerateCode;
+
+	/** True if the last global validity check returned that the project path is valid for creation */
+	bool bLastGlobalValidityCheckSuccessful;
+
+	/** True if the last NameAndLocation validity check returned that the project path is valid for creation */
+	bool bLastNameAndLocationValidityCheckSuccessful;
+
+};
\ No newline at end of file
diff --git a/Engine/Source/Editor/GameProjectGeneration/Private/TemplateCategory.h b/Engine/Source/Editor/GameProjectGeneration/Private/TemplateCategory.h
--- a/Engine/Source/Editor/GameProjectGeneration/Private/TemplateCategory.h
+++ b/Engine/Source/Editor/GameProjectGeneration/Private/TemplateCategory.h
@@ -20,9 +20,6 @@ struct FTemplateCategory
 	/** A unique key for this category */
 	FName Key;
 
-	/** Is this a major or minor category? Controls where in the UI it will show up. */
-	bool IsMajor;
-
 	/** Is this an enterprise project? Will end up including Datasmith modules if true. */
 	bool IsEnterprise;
 };
diff --git a/Engine/Source/Editor/GameProjectGeneration/Private/TemplateItem.h b/Engine/Source/Editor/GameProjectGeneration/Private/TemplateItem.h
--- a/Engine/Source/Editor/GameProjectGeneration/Private/TemplateItem.h
+++ b/Engine/Source/Editor/GameProjectGeneration/Private/TemplateItem.h
@@ -31,4 +31,5 @@ struct FTemplateItem
 
 	bool		bIsEnterprise = false;
 	bool		bIsBlankTemplate = false;
+	bool		bThumbnailAsIcon = false;
 };
diff --git a/Engine/Source/Editor/GameProjectGeneration/Private/TemplateProjectDefs.cpp b/Engine/Source/Editor/GameProjectGeneration/Private/TemplateProjectDefs.cpp
--- a/Engine/Source/Editor/GameProjectGeneration/Private/TemplateProjectDefs.cpp
+++ b/Engine/Source/Editor/GameProjectGeneration/Private/TemplateProjectDefs.cpp
@@ -60,6 +60,7 @@ UTemplateProjectDefs::UTemplateProjectDefs(const FObjectInitializer& ObjectIniti
 	: Super(ObjectInitializer)
 {
 	bAllowProjectCreation = true;
+	bThumbnailAsIcon = false;
 	EditDetailLevelPreference = EFeaturePackDetailLevel::Standard;
 }
 
diff --git a/Engine/Source/Editor/GameProjectGeneration/Private/Tests/GameProjectAutomationTests.cpp b/Engine/Source/Editor/GameProjectGeneration/Private/Tests/GameProjectAutomationTests.cpp
--- a/Engine/Source/Editor/GameProjectGeneration/Private/Tests/GameProjectAutomationTests.cpp
+++ b/Engine/Source/Editor/GameProjectGeneration/Private/Tests/GameProjectAutomationTests.cpp
@@ -15,7 +15,7 @@
 #include "ProjectDescriptor.h"
 #include "Editor/AddContentDialog/Private/IContentSource.h"
 #include "GameProjectUtils.h"
-#include "Editor/GameProjectGeneration/Private/SNewProjectWizard.h"
+#include "Editor/GameProjectGeneration/Private/SProjectDialog.h"
 
 #include "DesktopPlatformModule.h"
 #include "Tests/AutomationTestSettings.h"
@@ -78,7 +78,7 @@ namespace GameProjectAutomationUtils
 	 * @param	OutMatchedProjects	Total projects matching criteria
 	 * @param	OutCreatedProjects	Total projects succesfully created
 	 */
-	static void CreateProjectSet(TMap<FName, TArray<TSharedPtr<FTemplateItem>> >& InTemplates, EHardwareClass::Type InTargetedHardware, EGraphicsPreset::Type InGraphicPreset, FName InCategory, bool bInCopyStarterContent, int32 &OutCreatedProjects, int32 &OutMatchedProjects)
+	static void CreateProjectSet(TMap<FName, TArray<TSharedPtr<FTemplateItem>> >& InTemplates, EHardwareClass InTargetedHardware, EGraphicsPreset InGraphicPreset, FName InCategory, bool bInCopyStarterContent, int32 &OutCreatedProjects, int32 &OutMatchedProjects)
 	{		
 		// If this is empty, it will use the same name for each project, otherwise it will create a project based on target platform and source template
 		FString TestRootFolder;// = "ProjectTests";
@@ -329,10 +329,7 @@ IMPLEMENT_SIMPLE_AUTOMATION_TEST(FCreateBPTemplateProjectAutomationTests, "Syste
  */
 bool FCreateBPTemplateProjectAutomationTests::RunTest(const FString& Parameters)
 {
-	TSharedPtr<SNewProjectWizard> NewProjectWizard;
-	NewProjectWizard = SNew(SNewProjectWizard);
-	
-	TMap<FName, TArray<TSharedPtr<FTemplateItem>> >& Templates = NewProjectWizard->FindTemplateProjects();
+	TMap<FName, TArray<TSharedPtr<FTemplateItem>> > Templates = SProjectDialog::FindTemplateProjects();
 	int32 OutMatchedProjectsDesk = 0;
 	int32 OutCreatedProjectsDesk = 0;
 	GameProjectAutomationUtils::CreateProjectSet(Templates, EHardwareClass::Desktop, EGraphicsPreset::Maximum, "Game", false, OutMatchedProjectsDesk, OutCreatedProjectsDesk);
@@ -357,10 +354,8 @@ IMPLEMENT_SIMPLE_AUTOMATION_TEST(FCreateCPPTemplateProjectAutomationTests, "Syst
  */
 bool FCreateCPPTemplateProjectAutomationTests::RunTest(const FString& Parameters)
 {
-	TSharedPtr<SNewProjectWizard> NewProjectWizard;
-	NewProjectWizard = SNew(SNewProjectWizard);
-	//return ;
-	TMap<FName, TArray<TSharedPtr<FTemplateItem>> >& Templates = NewProjectWizard->FindTemplateProjects();//GameProjectAutomationUtils::CreateTemplateList();
+	TMap<FName, TArray<TSharedPtr<FTemplateItem>> > Templates = SProjectDialog::FindTemplateProjects();
+
 	int32 OutMatchedProjectsDesk = 0;
 	int32 OutCreatedProjectsDesk = 0;
 	GameProjectAutomationUtils::CreateProjectSet(Templates, EHardwareClass::Desktop, EGraphicsPreset::Maximum, "Game", false, OutMatchedProjectsDesk, OutCreatedProjectsDesk);
diff --git a/Engine/Source/Editor/GameProjectGeneration/Public/GameProjectUtils.h b/Engine/Source/Editor/GameProjectGeneration/Public/GameProjectUtils.h
--- a/Engine/Source/Editor/GameProjectGeneration/Public/GameProjectUtils.h
+++ b/Engine/Source/Editor/GameProjectGeneration/Public/GameProjectUtils.h
@@ -34,8 +34,8 @@ struct FProjectInformation
 	TOptional<bool> bEnableXR;
 	TOptional<bool> bEnableRaytracing;
 
-	TOptional<EHardwareClass::Type> TargetedHardware;
-	TOptional<EGraphicsPreset::Type> DefaultGraphicsPerformance;
+	TOptional<EHardwareClass> TargetedHardware;
+	TOptional<EGraphicsPreset> DefaultGraphicsPerformance;
 
 	/** The name of the feature pack to use as starter content. Must be located under FeaturePacks\. */
 	FString StarterContent;
diff --git a/Engine/Source/Editor/HardwareTargeting/HardwareTargeting.Build.cs b/Engine/Source/Editor/HardwareTargeting/HardwareTargeting.Build.cs
--- a/Engine/Source/Editor/HardwareTargeting/HardwareTargeting.Build.cs
+++ b/Engine/Source/Editor/HardwareTargeting/HardwareTargeting.Build.cs
@@ -22,6 +22,7 @@ public class HardwareTargeting : ModuleRules
 				"EditorStyle",
 				"EditorFramework",
 				"UnrealEd",
+				"EditorWidgets",
 				"Settings",
 				"EngineSettings",
 			}
diff --git a/Engine/Source/Editor/HardwareTargeting/Private/HardwareTargetingModule.cpp b/Engine/Source/Editor/HardwareTargeting/Private/HardwareTargetingModule.cpp
--- a/Engine/Source/Editor/HardwareTargeting/Private/HardwareTargetingModule.cpp
+++ b/Engine/Source/Editor/HardwareTargeting/Private/HardwareTargetingModule.cpp
@@ -11,13 +11,13 @@
 #include "EditorStyleSet.h"
 #include "ISettingsModule.h"
 #include "Widgets/SWidget.h"
-#include "SDecoratedEnumCombo.h"
 
 #include "Engine/RendererSettings.h"
 #include "Widgets/SToolTip.h"
 #include "IDocumentation.h"
 #include "Settings/EditorProjectSettings.h"
 #include "SlateSettings.h"
+#include "SEnumCombo.h"
 
 #define LOCTEXT_NAMESPACE "HardwareTargeting"
 
@@ -148,8 +148,8 @@ public:
 	// IHardwareTargetingModule interface
 	virtual void ApplyHardwareTargetingSettings() override;
 	virtual TArray<FModifiedDefaultConfig> GetPendingSettingsChanges() override;
-	virtual TSharedRef<SWidget> MakeHardwareClassTargetCombo(FOnHardwareClassChanged OnChanged, TAttribute<EHardwareClass::Type> SelectedEnum, EOrientation Orientation) override;
-	virtual TSharedRef<SWidget> MakeGraphicsPresetTargetCombo(FOnGraphicsPresetChanged OnChanged, TAttribute<EGraphicsPreset::Type> SelectedEnum, EOrientation Orientation) override;
+	virtual TSharedRef<SWidget> MakeHardwareClassTargetCombo(FOnHardwareClassChanged OnChanged, TAttribute<EHardwareClass> SelectedEnum) override;
+	virtual TSharedRef<SWidget> MakeGraphicsPresetTargetCombo(FOnGraphicsPresetChanged OnChanged, TAttribute<EGraphicsPreset> SelectedEnum) override;
 	// End of IHardwareTargetingModule interface
 
 private:
@@ -277,38 +277,31 @@ void FHardwareTargetingModule::ApplyHardwareTargetingSettings()
 	}
 }
 
-TSharedRef<SWidget> FHardwareTargetingModule::MakeHardwareClassTargetCombo(FOnHardwareClassChanged OnChanged, TAttribute<EHardwareClass::Type> SelectedEnum, EOrientation Orientation)
+TSharedRef<SWidget> FHardwareTargetingModule::MakeHardwareClassTargetCombo(FOnHardwareClassChanged OnChanged, TAttribute<EHardwareClass> SelectedEnum)
 {
-	TArray<SDecoratedEnumCombo<EHardwareClass::Type>::FComboOption> HardwareClassInfo;
-	HardwareClassInfo.Add(SDecoratedEnumCombo<EHardwareClass::Type>::FComboOption(
-		EHardwareClass::Unspecified, FSlateIcon(FEditorStyle::GetStyleSetName(), "HardwareTargeting.HardwareUnspecified"), LOCTEXT("UnspecifiedCaption", "Unspecified"), false));
-	HardwareClassInfo.Add(SDecoratedEnumCombo<EHardwareClass::Type>::FComboOption(
-		EHardwareClass::Desktop, FSlateIcon(FEditorStyle::GetStyleSetName(), "HardwareTargeting.DesktopPlatform"), LOCTEXT("DesktopCaption", "Desktop / Console")));
-	HardwareClassInfo.Add(SDecoratedEnumCombo<EHardwareClass::Type>::FComboOption(
-		EHardwareClass::Mobile, FSlateIcon(FEditorStyle::GetStyleSetName(), "HardwareTargeting.MobilePlatform"), LOCTEXT("MobileCaption", "Mobile / Tablet")));
-
-	return SNew(SDecoratedEnumCombo<EHardwareClass::Type>, MoveTemp(HardwareClassInfo))
-		.SelectedEnum(SelectedEnum)
-		.OnEnumChanged(OnChanged)
-		.Orientation(Orientation)
-		.ToolTip(IDocumentation::Get()->CreateToolTip(LOCTEXT("HardwareClassTooltip", "Choose the overall class of hardware to target (desktop/console or mobile/tablet)."), NULL, TEXT("Shared/Editor/Settings/TargetHardware"), TEXT("HardwareClass")));
+	const UEnum* HardwareClassEnum = StaticEnum<EHardwareClass>();
+
+	return
+		SNew(SEnumComboBox, HardwareClassEnum)
+		.ContentPadding(FMargin(4.0f, 0.0f))
+		.ToolTip(IDocumentation::Get()->CreateToolTip(LOCTEXT("HardwareClassTooltip", "Choose the overall class of hardware to target (desktop/console or mobile/tablet)."), NULL, TEXT("Shared/Editor/Settings/TargetHardware"), TEXT("HardwareClass")))
+		.OnEnumSelectionChanged_Lambda([OnChanged](int32 NewSelection, ESelectInfo::Type) { OnChanged.ExecuteIfBound(static_cast<EHardwareClass>(NewSelection)); })
+		.CurrentValue_Lambda([SelectedEnum]() { return static_cast<int32>(SelectedEnum.Get()); });
+
+
+		
 }
 
-TSharedRef<SWidget> FHardwareTargetingModule::MakeGraphicsPresetTargetCombo(FOnGraphicsPresetChanged OnChanged, TAttribute<EGraphicsPreset::Type> SelectedEnum, EOrientation Orientation)
+TSharedRef<SWidget> FHardwareTargetingModule::MakeGraphicsPresetTargetCombo(FOnGraphicsPresetChanged OnChanged, TAttribute<EGraphicsPreset> SelectedEnum)
 {
-	TArray<SDecoratedEnumCombo<EGraphicsPreset::Type>::FComboOption> GraphicsPresetInfo;
-	GraphicsPresetInfo.Add(SDecoratedEnumCombo<EGraphicsPreset::Type>::FComboOption(
-		EGraphicsPreset::Unspecified, FSlateIcon(FEditorStyle::GetStyleSetName(), "HardwareTargeting.GraphicsUnspecified"), LOCTEXT("UnspecifiedCaption", "Unspecified"), false));
-	GraphicsPresetInfo.Add(SDecoratedEnumCombo<EGraphicsPreset::Type>::FComboOption(
-		EGraphicsPreset::Maximum, FSlateIcon(FEditorStyle::GetStyleSetName(), "HardwareTargeting.MaximumQuality"), LOCTEXT("MaximumCaption", "Maximum Quality")));
-	GraphicsPresetInfo.Add(SDecoratedEnumCombo<EGraphicsPreset::Type>::FComboOption(
-		EGraphicsPreset::Scalable, FSlateIcon(FEditorStyle::GetStyleSetName(), "HardwareTargeting.ScalableQuality"), LOCTEXT("ScalableCaption", "Scalable 3D or 2D")));
-
-	return SNew(SDecoratedEnumCombo<EGraphicsPreset::Type>, MoveTemp(GraphicsPresetInfo))
-		.SelectedEnum(SelectedEnum)
-		.OnEnumChanged(OnChanged)
-		.Orientation(Orientation)
-		.ToolTip(IDocumentation::Get()->CreateToolTip(LOCTEXT("GraphicsPresetTooltip", "Choose the graphical level to target (high-end only or scalable from low-end on up)."), NULL, TEXT("Shared/Editor/Settings/TargetHardware"), TEXT("GraphicalLevel")));
+	const UEnum* GraphicsPresetEnum = StaticEnum<EGraphicsPreset>();
+	
+	return
+		SNew(SEnumComboBox, GraphicsPresetEnum)
+		.ContentPadding(FMargin(4.0f, 0.0f))
+		.ToolTip(IDocumentation::Get()->CreateToolTip(LOCTEXT("GraphicsPresetTooltip", "Choose the graphical level to target (high-end only or scalable from low-end on up)."), NULL, TEXT("Shared/Editor/Settings/TargetHardware"), TEXT("GraphicalLevel")))
+		.OnEnumSelectionChanged_Lambda([OnChanged](int32 NewSelection, ESelectInfo::Type) { OnChanged.ExecuteIfBound(static_cast<EGraphicsPreset>(NewSelection)); })
+		.CurrentValue_Lambda([SelectedEnum]() { return static_cast<int32>(SelectedEnum.Get()); });
 }
 
 IHardwareTargetingModule& IHardwareTargetingModule::Get()
diff --git a/Engine/Source/Editor/HardwareTargeting/Public/HardwareTargetingModule.h b/Engine/Source/Editor/HardwareTargeting/Public/HardwareTargetingModule.h
--- a/Engine/Source/Editor/HardwareTargeting/Public/HardwareTargetingModule.h
+++ b/Engine/Source/Editor/HardwareTargeting/Public/HardwareTargetingModule.h
@@ -11,8 +11,8 @@
 
 class SWidget;
 
-DECLARE_DELEGATE_OneParam(FOnHardwareClassChanged, EHardwareClass::Type)
-DECLARE_DELEGATE_OneParam(FOnGraphicsPresetChanged, EGraphicsPreset::Type)
+DECLARE_DELEGATE_OneParam(FOnHardwareClassChanged, EHardwareClass)
+DECLARE_DELEGATE_OneParam(FOnGraphicsPresetChanged, EGraphicsPreset)
 
 /** Struct specifying pending changes to a settings object */
 struct FModifiedDefaultConfig
@@ -41,9 +41,9 @@ public:
 	virtual TArray<FModifiedDefaultConfig> GetPendingSettingsChanges() = 0;
 
 	/** Make a new combo box for choosing a hardware class target */
-	virtual TSharedRef<SWidget> MakeHardwareClassTargetCombo(FOnHardwareClassChanged OnChanged, TAttribute<EHardwareClass::Type> SelectedEnum, EOrientation Orientation) = 0;
+	virtual TSharedRef<SWidget> MakeHardwareClassTargetCombo(FOnHardwareClassChanged OnChanged, TAttribute<EHardwareClass> SelectedEnum) = 0;
 
 	/** Make a new combo box for choosing a graphics preference */
-	virtual TSharedRef<SWidget> MakeGraphicsPresetTargetCombo(FOnGraphicsPresetChanged OnChanged, TAttribute<EGraphicsPreset::Type> SelectedEnum, EOrientation Orientation) = 0;
+	virtual TSharedRef<SWidget> MakeGraphicsPresetTargetCombo(FOnGraphicsPresetChanged OnChanged, TAttribute<EGraphicsPreset> SelectedEnum) = 0;
 };
 
diff --git a/Engine/Source/Editor/HardwareTargeting/Public/HardwareTargetingSettings.h b/Engine/Source/Editor/HardwareTargeting/Public/HardwareTargetingSettings.h
--- a/Engine/Source/Editor/HardwareTargeting/Public/HardwareTargetingSettings.h
+++ b/Engine/Source/Editor/HardwareTargeting/Public/HardwareTargetingSettings.h
@@ -9,37 +9,31 @@
 
 /** Enum specifying a class of hardware */
 UENUM()
-namespace EHardwareClass
+enum class EHardwareClass : uint8
 {
-	enum Type
-	{
-		/** Unspecified, meaning no choice has been made yet */
-		Unspecified,
-
-		/** Desktop or console */
-		Desktop,
-		
-		/** Mobile or tablet */
-		Mobile
-	};
-}
+	/** Unspecified, meaning no choice has been made yet */
+	Unspecified UMETA(Hidden),
+
+	/** Desktop or console */
+	Desktop,
+
+	/** Mobile or tablet */
+	Mobile
+};
 
 /** Enum specifying a graphics preset preference */
 UENUM()
-namespace EGraphicsPreset
+enum class EGraphicsPreset : uint8
 {
-	enum Type
-	{
-		/** Unspecified, meaning no choice has been made yet */
-		Unspecified,
-
-		/** Maximum quality - High-end features default to enabled */
-		Maximum,
-		
-		/** Scalable quality - Some features are disabled by default but can be enabled based on the actual hardware */
-		Scalable
-	};
-}
+	/** Unspecified, meaning no choice has been made yet */
+	Unspecified UMETA(Hidden),
+
+	/** Maximum quality - High-end features default to enabled */
+	Maximum,
+
+	/** Scalable quality - Some features are disabled by default but can be enabled based on the actual hardware */
+	Scalable
+};
 
 /** Hardware targeting settings, stored in default config, per-project */
 UCLASS(config=Engine, defaultconfig)
@@ -49,19 +43,19 @@ class HARDWARETARGETING_API UHardwareTargetingSettings : public UObject
 
 	/** Enum specifying which class of hardware this game is targeting */
 	UPROPERTY(config, EditAnywhere, category=None)
-	TEnumAsByte<EHardwareClass::Type> TargetedHardwareClass;
+	EHardwareClass TargetedHardwareClass;
 	
 	/** Enum that is set to TargetedHardwareClass when the settings have been successfully applied */
 	UPROPERTY(config)
-	TEnumAsByte<EHardwareClass::Type> AppliedTargetedHardwareClass;
+	EHardwareClass AppliedTargetedHardwareClass;
 
 	/** Enum specifying a graphics preset to use for this game */
 	UPROPERTY(config, EditAnywhere, category=None)
-	TEnumAsByte<EGraphicsPreset::Type> DefaultGraphicsPerformance;
+	EGraphicsPreset DefaultGraphicsPerformance;
 
 	/** Enum that is set to DefaultGraphicsPerformance when the settings have been successfully applied */
 	UPROPERTY(config)
-	TEnumAsByte<EGraphicsPreset::Type> AppliedDefaultGraphicsPerformance;
+	EGraphicsPreset AppliedDefaultGraphicsPerformance;
 
 	/** Check if these settings have any pending changes that require action */
 	bool HasPendingChanges() const;
diff --git a/Engine/Source/Editor/HardwareTargeting/Public/SDecoratedEnumCombo.h b/Engine/Source/Editor/HardwareTargeting/Public/SDecoratedEnumCombo.h
deleted file mode 100644
diff --git a/Engine/Source/Editor/MainFrame/Private/Frame/MainFrameActions.cpp b/Engine/Source/Editor/MainFrame/Private/Frame/MainFrameActions.cpp
--- a/Engine/Source/Editor/MainFrame/Private/Frame/MainFrameActions.cpp
+++ b/Engine/Source/Editor/MainFrame/Private/Frame/MainFrameActions.cpp
@@ -298,14 +298,12 @@ void FMainFrameActionCallbacks::SaveAll()
 	FEditorFileUtils::SaveDirtyPackages( bPromptUserToSave, bSaveMapPackages, bSaveContentPackages, bFastSave, bNotifyNoPackagesSaved, bCanBeDeclined );
 }
 
-TArray<FString> FMainFrameActionCallbacks::ProjectNames;
+TArray<FRecentProjectFile> FMainFrameActionCallbacks::RecentProjects;
 
 void FMainFrameActionCallbacks::CacheProjectNames()
 {
-	ProjectNames.Empty();
-
 	// The switch project menu is filled with recently opened project files
-	ProjectNames = GetDefault<UEditorSettings>()->RecentlyOpenedProjectFiles;
+	RecentProjects = GetDefault<UEditorSettings>()->RecentlyOpenedProjectFiles;
 }
 
 void FMainFrameActionCallbacks::NewProject( bool bAllowProjectOpening, bool bAllowProjectCreate )
@@ -435,17 +433,17 @@ void FMainFrameActionCallbacks::ZipUpProject()
 //	FModuleManager::LoadModuleChecked<ILocalizationDashboardModule>("LocalizationDashboard").Show();
 //}
 
-void FMainFrameActionCallbacks::SwitchProjectByIndex( int32 ProjectIndex )
+void FMainFrameActionCallbacks::SwitchProjectByIndex(int32 ProjectIndex)
 {
-	FUnrealEdMisc::Get().SwitchProject( ProjectNames[ ProjectIndex ] );
+	FUnrealEdMisc::Get().SwitchProject(RecentProjects[ProjectIndex].ProjectName);
 }
 
 void FMainFrameActionCallbacks::SwitchProject(const FString& GameOrProjectFileName)
 {
-	FUnrealEdMisc::Get().SwitchProject( GameOrProjectFileName );
+	FUnrealEdMisc::Get().SwitchProject(GameOrProjectFileName);
 }
 
-void FMainFrameActionCallbacks::OpenBackupDirectory( FString BackupFile )
+void FMainFrameActionCallbacks::OpenBackupDirectory(FString BackupFile)
 {
 	FPlatformProcess::LaunchFileInDefaultExternalApplication(*FPaths::GetPath(FPaths::ConvertRelativePathToFull(BackupFile)));
 }
@@ -486,12 +484,12 @@ bool FMainFrameActionCallbacks::FullScreen_IsChecked()
 
 bool FMainFrameActionCallbacks::CanSwitchToProject( int32 InProjectIndex )
 {
-	if (FApp::HasProjectName() && ProjectNames[InProjectIndex].StartsWith(FApp::GetProjectName()))
+	if (FApp::HasProjectName() && RecentProjects[InProjectIndex].ProjectName.StartsWith(FApp::GetProjectName()))
 	{
 		return false;
 	}
 
-	if ( FPaths::IsProjectFilePathSet() && ProjectNames[ InProjectIndex ] == FPaths::GetProjectFilePath() )
+	if (FPaths::IsProjectFilePathSet() && RecentProjects[InProjectIndex].ProjectName == FPaths::GetProjectFilePath())
 	{
 		return false;
 	}
diff --git a/Engine/Source/Editor/MainFrame/Private/Frame/MainFrameActions.h b/Engine/Source/Editor/MainFrame/Private/Frame/MainFrameActions.h
--- a/Engine/Source/Editor/MainFrame/Private/Frame/MainFrameActions.h
+++ b/Engine/Source/Editor/MainFrame/Private/Frame/MainFrameActions.h
@@ -9,7 +9,7 @@
 #include "Framework/Commands/Commands.h"
 
 class FUICommandList;
-
+struct FRecentProjectFile;
 
 /**
  * Unreal editor main frame actions
@@ -234,7 +234,7 @@ public:
 public:
 
 	// List of projects that the user can switch to.
-	static TArray<FString> ProjectNames;
+	static TArray<FRecentProjectFile> RecentProjects;
 
 protected:
 
diff --git a/Engine/Source/Editor/MainFrame/Private/MainFrameModule.h b/Engine/Source/Editor/MainFrame/Private/MainFrameModule.h
--- a/Engine/Source/Editor/MainFrame/Private/MainFrameModule.h
+++ b/Engine/Source/Editor/MainFrame/Private/MainFrameModule.h
@@ -191,7 +191,7 @@ protected:
 public:
 
 	/** Get the size of the project browser window */
-	static FVector2D GetProjectBrowserWindowSize() { return FVector2D(1100, 740); }
+	static FVector2D GetProjectBrowserWindowSize() { return FVector2D(1190, 733); }
 
 private:
 
diff --git a/Engine/Source/Editor/MainFrame/Private/Menus/MainMenu.cpp b/Engine/Source/Editor/MainFrame/Private/Menus/MainMenu.cpp
--- a/Engine/Source/Editor/MainFrame/Private/Menus/MainMenu.cpp
+++ b/Engine/Source/Editor/MainFrame/Private/Menus/MainMenu.cpp
@@ -499,7 +499,7 @@ void FMainMenu::RegisterRecentFileAndExitMenuItems()
 
 	{
 		FToolMenuSection& Section = MainTabFileMenu->AddSection("FileRecentFiles");
-		if (GetDefault<UEditorStyleSettings>()->bShowProjectMenus && FMainFrameActionCallbacks::ProjectNames.Num() > 0)
+		if (GetDefault<UEditorStyleSettings>()->bShowProjectMenus && FMainFrameActionCallbacks::RecentProjects.Num() > 0)
 		{
 			Section.AddSubMenu(
 				"RecentProjects",
diff --git a/Engine/Source/Editor/MainFrame/Private/Menus/RecentProjectsMenu.h b/Engine/Source/Editor/MainFrame/Private/Menus/RecentProjectsMenu.h
--- a/Engine/Source/Editor/MainFrame/Private/Menus/RecentProjectsMenu.h
+++ b/Engine/Source/Editor/MainFrame/Private/Menus/RecentProjectsMenu.h
@@ -7,6 +7,7 @@
 #include "ToolMenus.h"
 #include "Frame/MainFrameActions.h"
 #include "HAL/FileManager.h"
+#include "Settings/EditorSettings.h"
 
 #define LOCTEXT_NAMESPACE "FRecentProjectsMenu"
 
@@ -24,10 +25,10 @@ public:
 	 */
 	static void MakeMenu( UToolMenu* Menu )
 	{
-		for ( int32 ProjectIndex = 0; ProjectIndex < FMainFrameActionCallbacks::ProjectNames.Num() && ProjectIndex < FMainFrameCommands::Get().SwitchProjectCommands.Num(); ++ProjectIndex )
+		for ( int32 ProjectIndex = 0; ProjectIndex < FMainFrameActionCallbacks::RecentProjects.Num() && ProjectIndex < FMainFrameCommands::Get().SwitchProjectCommands.Num(); ++ProjectIndex )
 		{
 			// If it is a project file, display the filename without extension. Otherwise just display the project name.
-			const FString& ProjectName = FMainFrameActionCallbacks::ProjectNames[ ProjectIndex ];
+			const FString& ProjectName = FMainFrameActionCallbacks::RecentProjects[ ProjectIndex ].ProjectName;
 
 			if (( IFileManager::Get().FileSize(*ProjectName) <= 0 ) ||
 				( FPaths::GetProjectFilePath() == ProjectName ))
diff --git a/Engine/Source/Editor/UnrealEd/Private/EditorEngine.cpp b/Engine/Source/Editor/UnrealEd/Private/EditorEngine.cpp
--- a/Engine/Source/Editor/UnrealEd/Private/EditorEngine.cpp
+++ b/Engine/Source/Editor/UnrealEd/Private/EditorEngine.cpp
@@ -6645,18 +6645,22 @@ void UEditorEngine::OnLevelRemovedFromWorld(ULevel* InLevel, UWorld* InWorld)
 
 void UEditorEngine::UpdateRecentlyLoadedProjectFiles()
 {
-	if ( FPaths::IsProjectFilePathSet() )
+	if (FPaths::IsProjectFilePathSet())
 	{
+		FDateTime CurrentTime = FDateTime::UtcNow();
+
 		const FString AbsoluteProjectPath = IFileManager::Get().ConvertToAbsolutePathForExternalAppForRead(*FPaths::GetProjectFilePath());
 		// Update the recently loaded project files. Move this project file to the front of the list
-		TArray<FString>& RecentlyOpenedProjectFiles = GetMutableDefault<UEditorSettings>()->RecentlyOpenedProjectFiles;
-		RecentlyOpenedProjectFiles.Remove( AbsoluteProjectPath );
-		RecentlyOpenedProjectFiles.Insert( AbsoluteProjectPath, 0 );
+		TArray<FRecentProjectFile>& RecentlyOpenedProjectFiles = GetMutableDefault<UEditorSettings>()->RecentlyOpenedProjectFiles;
+
+		FRecentProjectFile MostRecentProject(AbsoluteProjectPath, CurrentTime);
+		RecentlyOpenedProjectFiles.Remove(MostRecentProject);
+		RecentlyOpenedProjectFiles.Insert(MostRecentProject, 0);
 
 		// Trim any project files that do not have the current game project file extension
 		for ( int32 FileIdx = RecentlyOpenedProjectFiles.Num() - 1; FileIdx >= 0; --FileIdx )
 		{
-			const FString FileExtension = FPaths::GetExtension(RecentlyOpenedProjectFiles[FileIdx]);
+			const FString FileExtension = FPaths::GetExtension(RecentlyOpenedProjectFiles[FileIdx].ProjectName);
 			if ( FileExtension != FProjectDescriptor::GetExtension() )
 			{
 				RecentlyOpenedProjectFiles.RemoveAt(FileIdx, 1);
@@ -6664,7 +6668,7 @@ void UEditorEngine::UpdateRecentlyLoadedProjectFiles()
 		}
 
 		// Trim the list in case we have more than the max
-		const int32 MaxRecentProjectFiles = 1024;
+		const int32 MaxRecentProjectFiles = 10;
 		if ( RecentlyOpenedProjectFiles.Num() > MaxRecentProjectFiles )
 		{
 			RecentlyOpenedProjectFiles.RemoveAt(MaxRecentProjectFiles, RecentlyOpenedProjectFiles.Num() - MaxRecentProjectFiles);
diff --git a/Engine/Source/Editor/UnrealEd/Public/Settings/EditorSettings.h b/Engine/Source/Editor/UnrealEd/Public/Settings/EditorSettings.h
--- a/Engine/Source/Editor/UnrealEd/Public/Settings/EditorSettings.h
+++ b/Engine/Source/Editor/UnrealEd/Public/Settings/EditorSettings.h
@@ -10,6 +10,40 @@
 #include "Engine/EngineTypes.h"
 #include "EditorSettings.generated.h"
 
+
+USTRUCT()
+struct FRecentProjectFile
+{
+	GENERATED_BODY()
+
+	/** Path to the project */
+	UPROPERTY(config)
+	FString ProjectName;
+
+	/** Timestamp of the last time the editor opened this project */
+	UPROPERTY(config)
+	FDateTime LastOpenTime;
+
+	FRecentProjectFile()
+	{}
+
+	FRecentProjectFile(const FString& InProjectName, FDateTime InLastOpenTime)
+		: ProjectName(InProjectName)
+		, LastOpenTime(InLastOpenTime)
+	{}
+
+	bool operator==(const FRecentProjectFile& Other) const
+	{
+		return ProjectName == Other.ProjectName;
+	}
+
+	bool operator==(const FString& OtherProjectName) const
+	{
+		return ProjectName == OtherProjectName;
+	}
+};
+
+
 UCLASS(config=EditorSettings)
 class UNREALED_API UEditorSettings : public UObject
 {
@@ -41,7 +75,7 @@ class UNREALED_API UEditorSettings : public UObject
 
 	/** Game project files that were recently opened in the editor */
 	UPROPERTY(config)
-	TArray<FString> RecentlyOpenedProjectFiles;
+	TArray<FRecentProjectFile> RecentlyOpenedProjectFiles;
 
 	/** The paths of projects created with the new project wizard. This is used to populate the "Path" field of the new project dialog. */
 	UPROPERTY(config)
diff --git a/Engine/Source/Runtime/SlateCore/Private/Styling/StarshipCoreStyle.cpp b/Engine/Source/Runtime/SlateCore/Private/Styling/StarshipCoreStyle.cpp
--- a/Engine/Source/Runtime/SlateCore/Private/Styling/StarshipCoreStyle.cpp
+++ b/Engine/Source/Runtime/SlateCore/Private/Styling/StarshipCoreStyle.cpp
@@ -295,7 +295,10 @@ TSharedRef<ISlateStyle> FStarshipCoreStyle::Create()
 		Style->Set("Icons.Download", new IMAGE_BRUSH("Icons/icon_Downloads_16x", Icon16x16));
 
 		Style->Set("Icons.Error", new IMAGE_BRUSH_SVG("Starship/Common/alert-circle", Icon16x16));
+		Style->Set("Icons.ErrorWithColor", new IMAGE_BRUSH_SVG("Starship/Common/alert-circle", Icon16x16, FStyleColors::Error));
+
 		Style->Set("Icons.Warning", new IMAGE_BRUSH_SVG("Starship/Common/alert-triangle", Icon16x16));
+		Style->Set("Icons.WarningWithColor", new IMAGE_BRUSH_SVG("Starship/Common/alert-triangle", Icon16x16, FStyleColors::Warning));
 
 		Style->Set("Icons.box-perspective", new IMAGE_BRUSH_SVG("Starship/Common/box-perspective", Icon16x16));
 		Style->Set("Icons.cylinder", new IMAGE_BRUSH_SVG("Starship/Common/cylinder", Icon16x16));
@@ -352,6 +355,9 @@ TSharedRef<ISlateStyle> FStarshipCoreStyle::Create()
 
 		Style->Set("Icons.DragHandle", new IMAGE_BRUSH_SVG("Starship/Common/drag-handle", Icon16x16));
 
+		Style->Set("Icons.Refresh", new IMAGE_BRUSH_SVG("Starship/Common/refresh", Icon16x16));
+
+
 	}
 
 	// Tool panels
@@ -1114,6 +1120,7 @@ void FStarshipCoreStyle::SetupButtonStyles(TSharedRef<FStyle>& Style)
 		Style->Set("SimpleButton", SimpleButton);
 		Style->Set("SecondaryButton", SecondaryButton);
 
+
 		Style->Set("DialogButtonText", FTextBlockStyle(NormalText).SetFont(FStyleFonts::Get().NormalBold).SetTransformPolicy(ETextTransformPolicy::ToUpper));
 
 
@@ -1121,6 +1128,21 @@ void FStarshipCoreStyle::SetupButtonStyles(TSharedRef<FStyle>& Style)
 		AddFont.LetterSpacing = 200;
 
 		Style->Set("SmallButtonText", FTextBlockStyle(NormalText).SetFont(AddFont).SetTransformPolicy(ETextTransformPolicy::ToUpper));
+
+		// A button that only responds to clicks and displays nothing
+		const FButtonStyle InvisibleButton = FButtonStyle()
+			.SetNormal(FSlateNoResource())
+			.SetHovered(FSlateNoResource())
+			.SetPressed(FSlateNoResource())
+			.SetDisabled(FSlateNoResource())
+			.SetNormalForeground(FLinearColor::White)
+			.SetHoveredForeground(FLinearColor::White)
+			.SetPressedForeground(FLinearColor::White)
+			.SetDisabledForeground(FLinearColor::White)
+			.SetNormalPadding(FMargin(0))
+			.SetPressedPadding(FMargin(0));
+
+		Style->Set("InvisibleButton", InvisibleButton);
 	}
 }
 
diff --git a/Templates/TP_AEC_ArchvisBP/Config/TemplateDefs.ini b/Templates/TP_AEC_ArchvisBP/Config/TemplateDefs.ini
--- a/Templates/TP_AEC_ArchvisBP/Config/TemplateDefs.ini
+++ b/Templates/TP_AEC_ArchvisBP/Config/TemplateDefs.ini
@@ -25,6 +25,9 @@ LocalizedDescriptions=(Language="ja",Text="建築ビジュアライゼーショ
 ; The sort key to use for the template to bypass the default alphabetical sorting.
 SortKey="1"
 
+; This project uses a thumbnail as the project icon
+bThumbnailAsIcon=true
+
 ; A list of folders to avoid copying when creating a project from this template. 
 ; These should contain files only of concern to the template itself and not desired in the user's project.
 FoldersToIgnore=Media
diff --git a/Templates/TP_AEC_BlankBP/Config/TemplateDefs.ini b/Templates/TP_AEC_BlankBP/Config/TemplateDefs.ini
--- a/Templates/TP_AEC_BlankBP/Config/TemplateDefs.ini
+++ b/Templates/TP_AEC_BlankBP/Config/TemplateDefs.ini
@@ -27,6 +27,9 @@ LocalizedDescriptions=(Language="ja", Text="最小限の機能とコンテンツ
 ; The sort key to use for the template to bypass the default alphabetical sorting.
 SortKey="0"
 
+; This project uses a thumbnail as the project icon
+bThumbnailAsIcon=true
+
 ; A list of folders to avoid copying when creating a project from this template. 
 ; These should contain files only of concern to the template itself and not desired in the user's project.
 FoldersToIgnore=Media
diff --git a/Templates/TP_CollaborativeBP/Config/TemplateDefs.ini b/Templates/TP_CollaborativeBP/Config/TemplateDefs.ini
--- a/Templates/TP_CollaborativeBP/Config/TemplateDefs.ini
+++ b/Templates/TP_CollaborativeBP/Config/TemplateDefs.ini
@@ -22,6 +22,9 @@ LocalizedDescriptions=(Language="ko", Text="이 프로젝트 템플릿은 데스
 ; The sort key to use for the template to bypass the default alphabetical sorting.
 SortKey="3"
 
+; This project uses a thumbnail as the project icon
+bThumbnailAsIcon=true
+
 ; A list of folders to avoid copying when creating a project from this template. 
 ; These should contain files only of concern to the template itself and not desired in the user's project.
 FoldersToIgnore=Media
diff --git a/Templates/TP_DMXBP/Config/TemplateDefs.ini b/Templates/TP_DMXBP/Config/TemplateDefs.ini
--- a/Templates/TP_DMXBP/Config/TemplateDefs.ini
+++ b/Templates/TP_DMXBP/Config/TemplateDefs.ini
@@ -14,6 +14,9 @@ LocalizedDescriptions=(Language="ja",Text="このテンプレートは DMX プ
 LocalizedDisplayNames=(Language="zh-Hans",Text="DMX")
 LocalizedDescriptions=(Language="zh-Hans",Text="该模板专门用于现场活动的视效预览和需要DMX协议的虚拟制片流程。它包含了寻址、配接和控制代理照明设备，以及在UE内外录制和回放现场DMX数据流的示例。\n\n它是一个很好的起点，可让你为使用专业照明控制台的现场演出制作视效预览。它可以让你在虚幻内控制外部启用了DMX的实体设备，或是在你使用了DMX的项目内部触发自定义事件。它还展示了实验性的像素映射功能和额外的可控制效果，如冷焰火、烟花和喷泉。\n\n")
 
+; This project uses a thumbnail as the project icon
+bThumbnailAsIcon=false
+
 ClassTypes=Actor, Blueprint
 AssetTypes=Static Mesh, Texture, DMX Library, Level Sequence
 
diff --git a/Templates/TP_HololensBP/Config/TemplateDefs.ini b/Templates/TP_HololensBP/Config/TemplateDefs.ini
--- a/Templates/TP_HololensBP/Config/TemplateDefs.ini
+++ b/Templates/TP_HololensBP/Config/TemplateDefs.ini
@@ -22,6 +22,9 @@ LocalizedDescriptions=(Language="ko", Text="이 프로젝트 템플릿은 Holole
 ; The sort key to use for the template to bypass the default alphabetical sorting.
 SortKey="3"
 
+; This project uses a thumbnail as the project icon
+bThumbnailAsIcon=true
+
 ; A list of folders to avoid copying when creating a project from this template. 
 ; These should contain files only of concern to the template itself and not desired in the user's project.
 FoldersToIgnore=Media
diff --git a/Templates/TP_ME_BlankBP/Config/TemplateDefs.ini b/Templates/TP_ME_BlankBP/Config/TemplateDefs.ini
--- a/Templates/TP_ME_BlankBP/Config/TemplateDefs.ini
+++ b/Templates/TP_ME_BlankBP/Config/TemplateDefs.ini
@@ -24,6 +24,9 @@ LocalizedDescriptions=(Language="ja", Text="最小限の機能が備わった Bl
 ; The sort key to use for the template to bypass the default alphabetical sorting.
 SortKey="_1"
 
+; This project uses a thumbnail as the project icon
+bThumbnailAsIcon=true
+
 ; A list of folders to avoid copying when creating a project from this template. 
 ; These should contain files only of concern to the template itself and not desired in the user's project.
 FoldersToIgnore=Media
diff --git a/Templates/TP_ME_VProdBP/Config/TemplateDefs.ini b/Templates/TP_ME_VProdBP/Config/TemplateDefs.ini
--- a/Templates/TP_ME_VProdBP/Config/TemplateDefs.ini
+++ b/Templates/TP_ME_VProdBP/Config/TemplateDefs.ini
@@ -24,6 +24,9 @@ LocalizedDescriptions=(Language="ja", Text="この Virtual Production プロジ
 ; The sort key to use for the template to bypass the default alphabetical sorting.
 SortKey="_2"
 
+; This project uses a thumbnail as the project icon
+bThumbnailAsIcon=true
+
 ; A list of folders to avoid copying when creating a project from this template. 
 ; These should contain files only of concern to the template itself and not desired in the user's project.
 FoldersToIgnore=Media
diff --git a/Templates/TP_PhotoStudioBP/Config/TemplateDefs.ini b/Templates/TP_PhotoStudioBP/Config/TemplateDefs.ini
--- a/Templates/TP_PhotoStudioBP/Config/TemplateDefs.ini
+++ b/Templates/TP_PhotoStudioBP/Config/TemplateDefs.ini
@@ -24,6 +24,9 @@ LocalizedDescriptions=(Language="ja", Text="このテンプレートでは、製
 ; The sort key to use for the template to bypass the default alphabetical sorting.
 SortKey="1"
 
+; This project uses a thumbnail as the project icon
+bThumbnailAsIcon=true
+
 ; A list of folders to avoid copying when creating a project from this template. 
 ; These should contain files only of concern to the template itself and not desired in the user's project.
 FoldersToIgnore=Media
diff --git a/Templates/TP_ProductConfigBP/Config/TemplateDefs.ini b/Templates/TP_ProductConfigBP/Config/TemplateDefs.ini
--- a/Templates/TP_ProductConfigBP/Config/TemplateDefs.ini
+++ b/Templates/TP_ProductConfigBP/Config/TemplateDefs.ini
@@ -22,6 +22,9 @@ LocalizedDescriptions=(Language="ja", Text="このプロジェクト テンプ
 ; The sort key to use for the template to bypass the default alphabetical sorting.
 SortKey="3"
 
+; This project uses a thumbnail as the project icon
+bThumbnailAsIcon=true
+
 ; A list of folders to avoid copying when creating a project from this template. 
 ; These should contain files only of concern to the template itself and not desired in the user's project.
 FoldersToIgnore=Media
diff --git a/Templates/TemplateCategories.ini b/Templates/TemplateCategories.ini
--- a/Templates/TemplateCategories.ini
+++ b/Templates/TemplateCategories.ini
@@ -12,10 +12,10 @@
 ; "Icon": A relative reference to the icon to display for the category - should be at minimum 128x128 to avoid being blurry.
 ; "IsMajorCategory": Is this a major, top-level category? Major categories are displayed as full rows, eg. the "Game" category.
 
-Categories=(Key="Games", LocalizedDisplayNames=((Language="en",Text="Games"),(Language="zh-Hans",Text="游戏"),(Language="ja",Text="ゲーム"),(Language="ko",Text="게임")), LocalizedDescriptions=((Language="en",Text="Start your game development journey with one of our key classes, levels, and examples."),(Language="zh-Hans", Text="使用我们的关键职业、关卡和示例开始你的游戏开发之旅。"),(Language="ja",Text="基本となるクラスやレベル、サンプルを一つ使ってゲーム開発をはじめましょう。"),(Language="ko",Text="언리얼 엔진 주요 클래스, 레벨 및 샘플을 활용해서 게임 개발을 시작해 보세요.")), Icon="Media/Games.png", IsMajorCategory=true)
+Categories=(Key="Games", LocalizedDisplayNames=((Language="en",Text="Games"),(Language="zh-Hans",Text="游戏"),(Language="ja",Text="ゲーム"),(Language="ko",Text="게임")), LocalizedDescriptions=((Language="en",Text="Start your game development journey with one of our key classes, levels, and examples."),(Language="zh-Hans", Text="使用我们的关键职业、关卡和示例开始你的游戏开发之旅。"),(Language="ja",Text="基本となるクラスやレベル、サンプルを一つ使ってゲーム開発をはじめましょう。"),(Language="ko",Text="언리얼 엔진 주요 클래스, 레벨 및 샘플을 활용해서 게임 개발을 시작해 보세요.")), Icon="Media/Games_2x.png", IsMajorCategory=true)
 
-Categories=(Key="ME", LocalizedDisplayNames=((Language="en",Text="Film, Television, and Live Events"),(Language="zh-Hans",Text="影视与现场活动"),(Language="ja",Text="映画、テレビ、ライブ イベント"),(Language="ko",Text="영화, TV 및 라이브 이벤트")), LocalizedDescriptions=((Language="en",Text="Choose from templates and examples for nDisplay, VR Scouting, and virtual production workflows."),(Language="zh-Hans",Text="选择适用于nDisplay、VR探查和虚拟制片工作流程的模板和示例。"),(Language="ja", Text="nDisplay やバーチャル スカウティング ツール、バーチャル プロダクション ワークフローのためのテンプレートおよびサンプルを選択することができます。"),(Language="ko",Text="nDisplay, VR 스카우팅 및 버추얼 프로덕션 워크플로를 위한 템플릿과 샘플을 선택해 보세요.")), Icon="Media/ME.png", IsMajorCategory=true)
+Categories=(Key="ME", LocalizedDisplayNames=((Language="en",Text="Film / Video &\nLive Events"),(Language="zh-Hans",Text="影视与现场活动"),(Language="ja",Text="映画、テレビ、ライブ イベント"),(Language="ko",Text="영화, TV 및 라이브 이벤트")), LocalizedDescriptions=((Language="en",Text="Choose from templates and examples for nDisplay, VR Scouting, and virtual production workflows."),(Language="zh-Hans",Text="选择适用于nDisplay、VR探查和虚拟制片工作流程的模板和示例。"),(Language="ja", Text="nDisplay やバーチャル スカウティング ツール、バーチャル プロダクション ワークフローのためのテンプレートおよびサンプルを選択することができます。"),(Language="ko",Text="nDisplay, VR 스카우팅 및 버추얼 프로덕션 워크플로를 위한 템플릿과 샘플을 선택해 보세요.")), Icon="Media/FilmVideoLive_2x.png", IsMajorCategory=true)
 
-Categories=(Key="AEC", LocalizedDisplayNames=((Language="en",Text="Architecture, Engineering, and Construction"),(Language="zh-Hans",Text="建筑、工程与施工"),(Language="ja",Text="建築、土木、建設"),(Language="ko",Text="건축, 엔지니어링 및 건설")), LocalizedDescriptions=((Language="en",Text="Select a starting point for multi-user design reviews, photorealistic architectural design visualizations, sunlight studies, or stylized renderings."),(Language="zh-Hans", Text="为多用户设计评审、照片级建筑设计可视化、阳光研究或风格化渲染选择一个起点。"),(Language="ja",Text="複数のユーザーによるレビューや、フォトリアルな建築デザインのビジュアライゼーション、太陽光の調査、スタイライズド レンダリングのための土台を選択することができます。"),(Language="ko",Text="멀티 유저 디자인 리뷰, 건축 디자인의 포토리얼한 시각화, 선라이트 스터디 또는 스타일라이즈드 렌더링을 시작해 보세요.")), Icon="Media/AEC.png", IsMajorCategory=true)
+Categories=(Key="AEC", LocalizedDisplayNames=((Language="en",Text="Architecture\nEngineering &\nConstruction"),(Language="zh-Hans",Text="建筑、工程与施工"),(Language="ja",Text="建築、土木、建設"),(Language="ko",Text="건축, 엔지니어링 및 건설")), LocalizedDescriptions=((Language="en",Text="Select a starting point for multi-user design reviews, photorealistic architectural design visualizations, sunlight studies, or stylized renderings."),(Language="zh-Hans", Text="为多用户设计评审、照片级建筑设计可视化、阳光研究或风格化渲染选择一个起点。"),(Language="ja",Text="複数のユーザーによるレビューや、フォトリアルな建築デザインのビジュアライゼーション、太陽光の調査、スタイライズド レンダリングのための土台を選択することができます。"),(Language="ko",Text="멀티 유저 디자인 리뷰, 건축 디자인의 포토리얼한 시각화, 선라이트 스터디 또는 스타일라이즈드 렌더링을 시작해 보세요.")), Icon="Media/AEC_2x.png", IsMajorCategory=true)
 
-Categories=(Key="MFG", LocalizedDisplayNames=((Language="en",Text="Automotive, Product Design, and Manufacturing"),(Language="zh-Hans",Text="汽车、产品设计和制造"),(Language="ja",Text="自動車、プロダクト デザイン、製造"),(Language="ko",Text="자동차, 제품 디자인 및 제조")), LocalizedDescriptions=((Language="en",Text="Find templates for multi-user design reviews, photobooth studio environments, and product configurators."),(Language="zh-Hans",Text="为多用户设计评审、Photobooth工作室环境和产品配置器寻找适合的模板。"),(Language="ja",Text="複数のユーザーによるレビューや、フォトブース スタジオの環境背景、製品コンフィギュレーターのためのテンプレートを見つけることができます。"),(Language="ko",Text="멀티 유저 디자인 리뷰, 포토부스 스튜디오 환경 및 제품 컨피규레이터를 위한 템플릿을 살펴보세요.")), Icon="Media/MFG.png", IsMajorCategory=true)
+Categories=(Key="MFG", LocalizedDisplayNames=((Language="en",Text="Automotive\nProduct Design &\nManufacturing"),(Language="zh-Hans",Text="汽车、产品设计和制造"),(Language="ja",Text="自動車、プロダクト デザイン、製造"),(Language="ko",Text="자동차, 제품 디자인 및 제조")), LocalizedDescriptions=((Language="en",Text="Find templates for multi-user design reviews, photobooth studio environments, and product configurators."),(Language="zh-Hans",Text="为多用户设计评审、Photobooth工作室环境和产品配置器寻找适合的模板。"),(Language="ja",Text="複数のユーザーによるレビューや、フォトブース スタジオの環境背景、製品コンフィギュレーターのためのテンプレートを見つけることができます。"),(Language="ko",Text="멀티 유저 디자인 리뷰, 포토부스 스튜디오 환경 및 제품 컨피규레이터를 위한 템플릿을 살펴보세요.")), Icon="Media/AutomotiveProduct_2x.png", IsMajorCategory=true)
