Fix project dialog static analysis

diff --git a/Engine/Source/Editor/GameProjectGeneration/Private/SProjectDialog.cpp b/Engine/Source/Editor/GameProjectGeneration/Private/SProjectDialog.cpp
--- a/Engine/Source/Editor/GameProjectGeneration/Private/SProjectDialog.cpp
+++ b/Engine/Source/Editor/GameProjectGeneration/Private/SProjectDialog.cpp
@@ -769,8 +769,6 @@ TSharedRef<SWidget> SProjectDialog::MakeHybridView(EProjectDialogModeMode Mode)
 	Templates = FindTemplateProjects();
 	SetDefaultProjectLocation();
 
-	Mode = EProjectDialogModeMode::Hybrid;
-
 	TemplateListView = SNew(STileView<TSharedPtr<FTemplateItem>>)
 		.ListItemsSource(&FilteredTemplateList)
 		.SelectionMode(ESelectionMode::Single)
