Move recent projects and custom template thumbnails to engine content so that they appear properly if no templates are synced

diff --git a/Engine/Source/Editor/EditorStyle/Private/StarshipStyle.cpp b/Engine/Source/Editor/EditorStyle/Private/StarshipStyle.cpp
--- a/Engine/Source/Editor/EditorStyle/Private/StarshipStyle.cpp
+++ b/Engine/Source/Editor/EditorStyle/Private/StarshipStyle.cpp
@@ -1624,8 +1624,7 @@ void FStarshipEditorStyle::FStyle::SetupGeneralStyles()
 			.SetSelectedTextColor( InvertedForeground )
 			);
 
-		Set( "GameProjectDialog.DefaultGameThumbnail", new IMAGE_BRUSH( "GameProjectDialog/default_game_thumbnail_128x", Icon128x128 ) );
-		Set( "GameProjectDialog.DefaultGameThumbnail.Small", new IMAGE_BRUSH( "GameProjectDialog/default_game_thumbnail", Icon128x128 ) );
+		Set( "GameProjectDialog.DefaultGameThumbnail", new IMAGE_BRUSH( "GameProjectDialog/default_game_thumbnail", Icon64x64 ) );
 		Set( "GameProjectDialog.BlankProjectThumbnail", new IMAGE_BRUSH( "GameProjectDialog/blank_project_thumbnail", Icon128x128 ) );
 		Set( "GameProjectDialog.BlankProjectPreview", new IMAGE_BRUSH( "GameProjectDialog/blank_project_preview", FVector2D(400, 200) ) );
 		Set( "GameProjectDialog.BasicCodeThumbnail", new IMAGE_BRUSH( "GameProjectDialog/basic_code_thumbnail", Icon128x128 ) );
diff --git a/Engine/Source/Editor/GameProjectGeneration/Private/SProjectBrowser.cpp b/Engine/Source/Editor/GameProjectGeneration/Private/SProjectBrowser.cpp
--- a/Engine/Source/Editor/GameProjectGeneration/Private/SProjectBrowser.cpp
+++ b/Engine/Source/Editor/GameProjectGeneration/Private/SProjectBrowser.cpp
@@ -155,7 +155,7 @@ public:
 							.VAlign(VAlign_Center)
 							[
 								SNew(SImage)
-								.Image(ProjectItem->ProjectThumbnail ? ProjectItem->ProjectThumbnail.Get() : FAppStyle::Get().GetBrush("GameProjectDialog.DefaultGameThumbnail.Small"))
+								.Image(ProjectItem->ProjectThumbnail ? ProjectItem->ProjectThumbnail.Get() : FAppStyle::Get().GetBrush("GameProjectDialog.DefaultGameThumbnail"))
 							]
 						]
 					]
@@ -542,22 +542,6 @@ bool SProjectBrowser::CanExecuteFindInExplorer() const
 	return SelectedProjectItem.IsValid();
 }
 
-
-const FSlateBrush* SProjectBrowser::GetProjectItemImage(TWeakPtr<FProjectItem> ProjectItem) const
-{
-	if ( ProjectItem.IsValid() )
-	{
-		TSharedPtr<FProjectItem> Item = ProjectItem.Pin();
-		if ( Item->ProjectThumbnail.IsValid() )
-		{
-			return Item->ProjectThumbnail.Get();
-		}
-	}
-	
-	return FEditorStyle::GetBrush("GameProjectDialog.DefaultGameThumbnail");
-}
-
-
 TSharedPtr<FProjectItem> SProjectBrowser::GetSelectedProjectItem() const
 {
 	TArray< TSharedPtr<FProjectItem> > SelectedItems = ProjectTileView->GetSelectedItems();
@@ -662,146 +646,6 @@ static TSharedPtr<FProjectItem> CreateProjectItem(const FString& ProjectFilename
 
 FReply SProjectBrowser::FindProjects()
 {
-#if 0
-	// Create a map of parent project folders to their category
-	TMap<FString, EProjectCategoryType> ProjectFilesToCategoryType;
-
-	// Find all the engine installations
-	TMap<FString, FString> EngineInstallations;
-	FDesktopPlatformModule::Get()->EnumerateEngineInstallations(EngineInstallations);
-
-	// Add projects from every branch that we know about
-	FString CurrentEngineIdentifier = FDesktopPlatformModule::Get()->GetCurrentEngineIdentifier();
-	for(TMap<FString, FString>::TConstIterator Iter(EngineInstallations); Iter; ++Iter)
-	{
-		TArray<FString> ProjectFiles;
-		if(FDesktopPlatformModule::Get()->EnumerateProjectsKnownByEngine(Iter.Key(), false, ProjectFiles))
-		{
-			for(int Idx = 0; Idx < ProjectFiles.Num(); Idx++)
-			{
-				ProjectFilesToCategoryType.Add(ProjectFiles[Idx], EProjectCategoryType::UserDefined);
-			}
-		}
-	}
-
-	// Add all the samples from the launcher
-	TArray<FString> LauncherSampleProjects;
-	FDesktopPlatformModule::Get()->EnumerateLauncherSampleProjects(LauncherSampleProjects);
-	for(int32 Idx = 0; Idx < LauncherSampleProjects.Num(); Idx++)
-	{
-		ProjectFilesToCategoryType.Add(LauncherSampleProjects[Idx], EProjectCategoryType::Sample);
-	}
-
-	// Add all the native project files we can find, and automatically filter them depending on their directory
-	FUProjectDictionary& DefaultProjectDictionary = FUProjectDictionary::GetDefault();
-	DefaultProjectDictionary.Refresh();
-	const TArray<FString> &NativeProjectFiles = DefaultProjectDictionary.GetProjectPaths();
-	for(int32 Idx = 0; Idx < NativeProjectFiles.Num(); Idx++)
-	{
-		if(!NativeProjectFiles[Idx].Contains(TEXT("/Templates/")))
-		{
-			const EProjectCategoryType ProjectCategoryType = NativeProjectFiles[Idx].Contains(TEXT("/Samples/")) ? EProjectCategoryType::Sample : EProjectCategoryType::UserDefined;
-			ProjectFilesToCategoryType.Add(NativeProjectFiles[Idx], ProjectCategoryType);
-		}
-	}
-
-	// Normalize all the filenames and make sure there are no duplicates
-	TMap<FString, EProjectCategoryType> AbsoluteProjectFilesToCategory;
-	for(TMap<FString, EProjectCategoryType>::TConstIterator Iter(ProjectFilesToCategoryType); Iter; ++Iter)
-	{
-		FString AbsoluteFile = IFileManager::Get().ConvertToAbsolutePathForExternalAppForRead(*Iter.Key());
-		AbsoluteProjectFilesToCategory.Add(AbsoluteFile, Iter.Value());
-	}
-
-	const FText MyProjectsCategoryName = LOCTEXT("MyProjectsCategoryName", "My Projects");
-	const FText SamplesCategoryName = LOCTEXT("SamplesCategoryName", "Samples");
-
-	// Add all the discovered projects to the list
-	for(TMap<FString, EProjectCategoryType>::TConstIterator Iter(AbsoluteProjectFilesToCategory); Iter; ++Iter)
-	{
-		const FString& ProjectFilename = *Iter.Key();
-		TSharedPtr<FProjectItem> NewProjectItem = CreateProjectItem(ProjectFilename);
-		if (NewProjectItem.IsValid())
-		{
-			FText ProjectCategory;
-			
-			if (NewProjectItem->Category.IsEmpty())
-			{
-				// Not category specified, so use the category for the detected project type
-				ProjectCategory = (Iter.Value() == EProjectCategoryType::Sample) ? SamplesCategoryName : MyProjectsCategoryName;
-			}
-			else
-			{
-				// Use the user defined category
-				ProjectCategory = NewProjectItem->Category;
-			}
-
-			AddProjectToCategory(NewProjectItem.ToSharedRef(), ProjectCategory);
-		}
-	}
-
-	// Make sure the category order is "My Projects", "Samples", then all remaining categories in alphabetical order
-	TSharedPtr<FProjectCategory> MyProjectsCategory;
-	TSharedPtr<FProjectCategory> SamplesCategory;
-
-	for ( int32 CategoryIdx = ProjectCategories.Num() - 1; CategoryIdx >= 0; --CategoryIdx )
-	{
-		TSharedRef<FProjectCategory> Category = ProjectCategories[CategoryIdx];
-		if ( Category->CategoryName.EqualTo(MyProjectsCategoryName) )
-		{
-			MyProjectsCategory = Category;
-			ProjectCategories.RemoveAt(CategoryIdx);
-		}
-		else if ( Category->CategoryName.EqualTo(SamplesCategoryName) )
-		{
-			SamplesCategory = Category;
-			ProjectCategories.RemoveAt(CategoryIdx);
-		}
-	}
-
-	// Sort categories
-	struct FCompareCategories
-	{
-		FORCEINLINE bool operator()( const TSharedRef<FProjectCategory>& A, const TSharedRef<FProjectCategory>& B ) const
-		{
-			return A->CategoryName.CompareToCaseIgnored(B->CategoryName) < 0;
-		}
-	};
-	ProjectCategories.Sort( FCompareCategories() );
-
-	// Now read the built-in categories (last added is first in the list)
-	if ( SamplesCategory.IsValid() )
-	{
-		ProjectCategories.Insert(SamplesCategory.ToSharedRef(), 0);
-	}
-	if ( MyProjectsCategory.IsValid() )
-	{
-		ProjectCategories.Insert(MyProjectsCategory.ToSharedRef(), 0);
-	}
-
-	// Sort each individual category
-	struct FCompareProjectItems
-	{
-		FORCEINLINE bool operator()( const TSharedPtr<FProjectItem>& A, const TSharedPtr<FProjectItem>& B ) const
-		{
-			return A->Name.CompareToCaseIgnored(B->Name) < 0;
-		}
-	};
-	for ( int32 CategoryIdx = 0; CategoryIdx < ProjectCategories.Num(); ++CategoryIdx )
-	{
-		TSharedRef<FProjectCategory> Category = ProjectCategories[CategoryIdx];
-		Category->ProjectItemsSource.Sort( FCompareProjectItems() );
-	}
-
-	PopulateFilteredProjects();
-
-	for (auto CategoryIt = ProjectCategories.CreateConstIterator(); CategoryIt; ++CategoryIt)
-	{
-		const TSharedRef<FProjectCategory>& Category = *CategoryIt;
-		ConstructCategory(CategoriesBox.ToSharedRef(), Category);
-	}
-#else
-
 	const FString LastSelectedProjectFile = CurrentSelectedProjectPath;
 
 	ProjectItemsSource.Empty();
@@ -810,21 +654,17 @@ FReply SProjectBrowser::FindProjects()
 
 	TArray<FRecentProjectFile> RecentProjects = GetDefault<UEditorSettings>()->RecentlyOpenedProjectFiles;
 
-	//TSet<FString> ExistingRecentProjects = RecentProjects;
-
 	TSet<FString> AllFoundProjectFiles;
 
-	//AllFoundProjectFiles.Append(RecentProjects);
-
 	// Find all the engine installations
 	TMap<FString, FString> EngineInstallations;
 	FDesktopPlatformModule::Get()->EnumerateEngineInstallations(EngineInstallations);
 
 	// Add projects from every branch that we know about
-	FString CurrentEngineIdentifier = FDesktopPlatformModule::Get()->GetCurrentEngineIdentifier();
 	for (TMap<FString, FString>::TConstIterator Iter(EngineInstallations); Iter; ++Iter)
 	{
 		TArray<FString> ProjectFiles;
+
 		if (FDesktopPlatformModule::Get()->EnumerateProjectsKnownByEngine(Iter.Key(), false, ProjectFiles))
 		{
 			AllFoundProjectFiles.Append(MoveTemp(ProjectFiles));
@@ -885,8 +725,6 @@ FReply SProjectBrowser::FindProjects()
 		ProjectTileView->SetSelection(FilteredProjectItemsSource[0], ESelectInfo::Direct);
 	}
 
-#endif
-
 	return FReply::Handled();
 }
 
diff --git a/Engine/Source/Editor/GameProjectGeneration/Private/SProjectBrowser.h b/Engine/Source/Editor/GameProjectGeneration/Private/SProjectBrowser.h
--- a/Engine/Source/Editor/GameProjectGeneration/Private/SProjectBrowser.h
+++ b/Engine/Source/Editor/GameProjectGeneration/Private/SProjectBrowser.h
@@ -38,7 +38,7 @@ public:
 
 	SLATE_BEGIN_ARGS(SProjectBrowser)
 	{}
-		SLATE_ARGUMENT(FProjectSelectionChanged, OnSelectionChanged);
+	SLATE_ARGUMENT(FProjectSelectionChanged, OnSelectionChanged);
 	SLATE_END_ARGS()
 
 public:
@@ -51,7 +51,7 @@ public:
 	 *
 	 * @param InArgs - The construction arguments.
 	 */
-	void Construct( const FArguments& InArgs );
+	void Construct(const FArguments& InArgs);
 
 	bool HasProjects() const;
 
@@ -77,10 +77,10 @@ public:
 protected:
 
 	/** Creates a row in the template list */
-	TSharedRef<ITableRow> MakeProjectViewWidget( TSharedPtr<FProjectItem> ProjectItem, const TSharedRef<STableViewBase>& OwnerTable );
+	TSharedRef<ITableRow> MakeProjectViewWidget(TSharedPtr<FProjectItem> ProjectItem, const TSharedRef<STableViewBase>& OwnerTable);
 
 	/** Create a tooltip for the given project item */
-	TSharedRef<SToolTip> MakeProjectToolTip( TSharedPtr<FProjectItem> ProjectItem ) const;
+	TSharedRef<SToolTip> MakeProjectToolTip(TSharedPtr<FProjectItem> ProjectItem) const;
 
 	/** Add information to the tooltip for this project item */
 	void AddToToolTipInfoBox(const TSharedRef<SVerticalBox>& InfoBox, const FText& Key, const FText& Value) const;
@@ -92,19 +92,17 @@ protected:
 	void ExecuteFindInExplorer() const;
 
 	/** Handler to check to see if a find in explorer command is allowed */
-	bool CanExecuteFindInExplorer() const;	
+	bool CanExecuteFindInExplorer() const;
 
-	/** Gets the image to display for the specified template */
-	const FSlateBrush* GetProjectItemImage( TWeakPtr<FProjectItem> ProjectItem ) const;
 
 	/** Gets the currently selected template item */
-	TSharedPtr<FProjectItem> GetSelectedProjectItem( ) const;
+	TSharedPtr<FProjectItem> GetSelectedProjectItem() const;
 
 	/** Populates ProjectItemsSource with projects found on disk */
-	FReply FindProjects( );
+	FReply FindProjects();
 
 	/** Opens the specified project file */
-	bool OpenProject( const FString& ProjectFile );
+	bool OpenProject(const FString& ProjectFile);
 
 	/** Populate the list of filtered project categories */
 	void PopulateFilteredProjects();
@@ -117,26 +115,21 @@ protected:
 	 *
 	 * @return  Returns whether the event was handled, along with other possible actions
 	 */
-	virtual FReply OnKeyDown( const FGeometry& MyGeometry, const FKeyEvent& InKeyEvent ) override;
+	virtual FReply OnKeyDown(const FGeometry& MyGeometry, const FKeyEvent& InKeyEvent) override;
 
 private:
-
-
-
 	/** Returns true if the user is allowed to open a project with the supplied name and path */
-	bool HandleOpenProjectButtonIsEnabled( ) const;
+	bool HandleOpenProjectButtonIsEnabled() const;
 
 	/** Called when a user double-clicks a project item in the project view */
-	void HandleProjectItemDoubleClick( TSharedPtr<FProjectItem> ProjectItem );
+	void HandleProjectItemDoubleClick(TSharedPtr<FProjectItem> ProjectItem);
 
 	/** Handler for when the selection changes in the project view */
-	void HandleProjectViewSelectionChanged( TSharedPtr<FProjectItem> ProjectItem, ESelectInfo::Type SelectInfo);
-
+	void HandleProjectViewSelectionChanged(TSharedPtr<FProjectItem> ProjectItem, ESelectInfo::Type SelectInfo);
 
 	/** Called when the text in the filter box is changed */
 	void OnFilterTextChanged(const FText& InText);
 
-
 	/** Get the visibility of the specified category */
 	EVisibility GetProjectCategoryVisibility(const TSharedRef<FProjectCategory> InCategory) const;
 
@@ -144,7 +137,7 @@ private:
 
 	/** Get the visibility of the "no projects" error */
 	EVisibility GetNoProjectsAfterFilterErrorVisibility() const;
-	
+
 	/** Get the filter text to highlight on items in the list */
 	FText GetItemHighlightText() const;
 
diff --git a/Engine/Source/Editor/GameProjectGeneration/Private/SProjectDialog.cpp b/Engine/Source/Editor/GameProjectGeneration/Private/SProjectDialog.cpp
--- a/Engine/Source/Editor/GameProjectGeneration/Private/SProjectDialog.cpp
+++ b/Engine/Source/Editor/GameProjectGeneration/Private/SProjectDialog.cpp
@@ -52,8 +52,11 @@ namespace NewProjectDialogDefs
 	constexpr float TemplateTileWidth = 102;
 	constexpr float ThumbnailSize = 64.0f, ThumbnailPadding = 5.f;
 	const FName DefaultCategoryName = "Games";
+	const FName BlankCategoryKey = "Default";
 }
 
+TUniquePtr<FSlateBrush> SProjectDialog::CustomTemplateBrush;
+
 class SMajorCategoryTile : public SCompoundWidget
 {
 	SLATE_BEGIN_ARGS(SMajorCategoryTile) {}
@@ -255,7 +258,7 @@ private:
 		{
 			return ItemPtr->Thumbnail.Get();
 		}
-		return FEditorStyle::GetBrush("GameProjectDialog.DefaultGameThumbnail.Small");
+		return FEditorStyle::GetBrush("GameProjectDialog.DefaultGameThumbnail");
 	}
 	
 };
@@ -492,19 +495,7 @@ void SProjectDialog::PopulateTemplateCategories()
 	TemplateCategories.Empty();
 	CurrentCategory.Reset();
 
-	FGameProjectGenerationModule::Get().GetAllTemplateCategories(TemplateCategories);
-
-	if (TemplateCategories.Num() == 0)
-	{
-		TSharedPtr<FTemplateCategory> DefaultCategory = MakeShared<FTemplateCategory>();
-		static const FName DefaultCategoryKey("Default");
-		DefaultCategory->Key = DefaultCategoryKey;
-		DefaultCategory->DisplayName = LOCTEXT("ProjectDialog_DefaultCategoryName", "Blank Project");
-		DefaultCategory->Description = LOCTEXT("ProjectDialog_DefaultCategoryDescription", "Create a new blank Unreal project.");
-		DefaultCategory->Icon = FAppStyle::Get().GetBrush("GameProjectDialog.DefaultGameThumbnail");
-
-		TemplateCategories.Add(DefaultCategory);
-	}
+	TemplateCategories = GetAllTemplateCategories();
 }
 
 TSharedRef<SWidget> SProjectDialog::MakeNewProjectDialogButtons()
@@ -1108,7 +1099,7 @@ TSharedRef<SWidget> SProjectDialog::MakeRecentProjectsTile()
 	RecentProjectsCategory->Key = "RecentProjects";
 	RecentProjectsCategory->IsEnterprise = false;
 
-	static const FName BrushName = *(FPaths::RootDir() / TEXT("Templates/Media") / TEXT("RecentProjects_2x.png"));
+	static const FName BrushName =  *(FAppStyle::Get().GetContentRootDir() / TEXT("/Starship/Projects/") / TEXT("RecentProjects_2x.png"));
 
 	RecentProjectsBrush = MakeUnique<FSlateDynamicImageBrush>(BrushName, FVector2D(300, 100));
 	RecentProjectsBrush->OutlineSettings.CornerRadii = FVector4(4, 4, 4, 4);
@@ -1800,12 +1791,21 @@ TArray<TSharedPtr<FTemplateCategory>> SProjectDialog::GetAllTemplateCategories()
 
 	if (AllTemplateCategories.Num() == 0)
 	{
-		TSharedPtr<FTemplateCategory> DefaultCategory = MakeShareable(new FTemplateCategory());
-		static const FName DefaultCategoryKey("Default");
-		DefaultCategory->Key = DefaultCategoryKey;
+		static const FName BrushName = *(FAppStyle::Get().GetContentRootDir() / TEXT("/Starship/Projects/") / TEXT("CustomTemplate_2x.png"));
+
+		if (!CustomTemplateBrush)
+		{
+			CustomTemplateBrush = MakeUnique<FSlateDynamicImageBrush>(BrushName, FVector2D(300, 100));
+			CustomTemplateBrush->OutlineSettings.CornerRadii = FVector4(4, 4, 4, 4);
+			CustomTemplateBrush->OutlineSettings.RoundingType = ESlateBrushRoundingType::FixedRadius;
+			CustomTemplateBrush->DrawAs = ESlateBrushDrawType::RoundedBox;
+		}
+
+		TSharedPtr<FTemplateCategory> DefaultCategory = MakeShared<FTemplateCategory>();
+		DefaultCategory->Key = NewProjectDialogDefs::BlankCategoryKey;
 		DefaultCategory->DisplayName = LOCTEXT("ProjectDialog_DefaultCategoryName", "Blank Project");
 		DefaultCategory->Description = LOCTEXT("ProjectDialog_DefaultCategoryDescription", "Create a new blank Unreal project.");
-		DefaultCategory->Icon = FEditorStyle::GetBrush("GameProjectDialog.DefaultGameThumbnail");
+		DefaultCategory->Icon = CustomTemplateBrush.Get();
 
 		AllTemplateCategories.Add(DefaultCategory);
 	}
@@ -1813,6 +1813,8 @@ TArray<TSharedPtr<FTemplateCategory>> SProjectDialog::GetAllTemplateCategories()
 	return AllTemplateCategories;
 }
 
+
+
 FText SProjectDialog::GetGlobalErrorLabelText() const
 {
 	if (!PersistentGlobalErrorLabelText.IsEmpty())
diff --git a/Engine/Source/Editor/GameProjectGeneration/Private/SProjectDialog.h b/Engine/Source/Editor/GameProjectGeneration/Private/SProjectDialog.h
--- a/Engine/Source/Editor/GameProjectGeneration/Private/SProjectDialog.h
+++ b/Engine/Source/Editor/GameProjectGeneration/Private/SProjectDialog.h
@@ -214,6 +214,8 @@ private:
 
 	TUniquePtr<FSlateBrush> RecentProjectsBrush;
 
+	static TUniquePtr<FSlateBrush> CustomTemplateBrush;
+
 	SVerticalBox::FSlot* ProjectOptionsSlot;
 	/** True if user has selected to copy starter content. */
 	bool bCopyStarterContent;
