Enable Lumen by default for new projects created off of templates

diff --git a/Engine/Source/Editor/GameProjectGeneration/Private/GameProjectUtils.cpp b/Engine/Source/Editor/GameProjectGeneration/Private/GameProjectUtils.cpp
--- a/Engine/Source/Editor/GameProjectGeneration/Private/GameProjectUtils.cpp
+++ b/Engine/Source/Editor/GameProjectGeneration/Private/GameProjectUtils.cpp
@@ -1814,6 +1814,7 @@ TOptional<FGuid> GameProjectUtils::CreateProjectFromTemplate(const FProjectInfor
 
 	AddHardwareConfigValues(InProjectInfo, ConfigValuesToSet);
 
+	AddLumenConfigValues(InProjectInfo, ConfigValuesToSet);
 	AddRaytracingConfigValues(InProjectInfo, ConfigValuesToSet);
 
 	AddDefaultMapConfigValues(InProjectInfo, ConfigValuesToSet);
