World Partition: honor templates bIsBlank property to automatically enable world partition in these projects (treat them as blank projects even if they are created from existing templates).

diff --git a/Engine/Source/Editor/GameProjectGeneration/Private/GameProjectUtils.cpp b/Engine/Source/Editor/GameProjectGeneration/Private/GameProjectUtils.cpp
--- a/Engine/Source/Editor/GameProjectGeneration/Private/GameProjectUtils.cpp
+++ b/Engine/Source/Editor/GameProjectGeneration/Private/GameProjectUtils.cpp
@@ -291,6 +291,31 @@ namespace
 				true /* ShouldReplaceExistingValue */);
 		}
 	}
+
+	/** Get the configuration values for enabling WorldPartition by default. */
+	void AddWorldPartitionConfigValues(const FProjectInformation& InProjectInfo, TArray<FTemplateConfigValue>& ConfigValues)
+	{
+		if (InProjectInfo.bIsBlankTemplate)
+		{
+			ConfigValues.Emplace(TEXT("DefaultEngine.ini"),
+				TEXT("/Script/WorldPartitionEditor.WorldPartitionEditorSettings"),
+				TEXT("bEnableWorldPartition"),
+				TEXT("True"),
+				true /* ShouldReplaceExistingValue */);
+
+			ConfigValues.Emplace(TEXT("DefaultEngine.ini"),
+				TEXT("/Script/WorldPartitionEditor.WorldPartitionEditorSettings"),
+				TEXT("bEnableConversionPrompt"),
+				TEXT("True"),
+				true /* ShouldReplaceExistingValue */);
+
+			ConfigValues.Emplace(TEXT("DefaultEngine.ini"),
+				TEXT("/Script/WorldPartitionEditor.WorldPartitionEditorSettings"),
+				TEXT("CommandletClass"),
+				TEXT("Class'/Script/UnrealEd.WorldPartitionConvertCommandlet"),
+				true /* ShouldReplaceExistingValue */);
+		}
+	}
 } // namespace <>
 
 FText FNewClassInfo::GetClassName() const
@@ -1835,6 +1860,7 @@ TOptional<FGuid> GameProjectUtils::CreateProjectFromTemplate(const FProjectInfor
 	AddRaytracingConfigValues(InProjectInfo, ConfigValuesToSet);
 	AddNewProjectDefaultShadowConfigValues(InProjectInfo, ConfigValuesToSet);
 	AddDefaultMapConfigValues(InProjectInfo, ConfigValuesToSet);
+	AddWorldPartitionConfigValues(InProjectInfo, ConfigValuesToSet);
 	
 	TemplateDefs->AddConfigValues(ConfigValuesToSet, TemplateName, ProjectName, InProjectInfo.bShouldGenerateCode);
 
@@ -2152,12 +2178,6 @@ bool GameProjectUtils::GenerateConfigFiles(const FProjectInformation& InProjectI
 			}
 		}
 
-		FileContents += LINE_TERMINATOR;
-		FileContents += TEXT("[/Script/WorldPartitionEditor.WorldPartitionEditorSettings]") LINE_TERMINATOR;
-		FileContents += TEXT("bEnableWorldPartition=True") LINE_TERMINATOR;
-		FileContents += TEXT("bEnableConversionPrompt=True") LINE_TERMINATOR;
-		FileContents += TEXT("CommandletClass=Class'/Script/UnrealEd.WorldPartitionConvertCommandlet'") LINE_TERMINATOR;
-
 		if (WriteOutputFile(DefaultEngineIniFilename, FileContents, OutFailReason))
 		{
 			OutCreatedFiles.Add(DefaultEngineIniFilename);
@@ -2172,6 +2192,7 @@ bool GameProjectUtils::GenerateConfigFiles(const FProjectInformation& InProjectI
 		AddLumenConfigValues(InProjectInfo, ConfigValuesToSet);
 		AddNewProjectDefaultShadowConfigValues(InProjectInfo, ConfigValuesToSet);
 		AddRaytracingConfigValues(InProjectInfo, ConfigValuesToSet);
+		AddWorldPartitionConfigValues(InProjectInfo, ConfigValuesToSet);
 
 		if (!SaveConfigValues(InProjectInfo, ConfigValuesToSet, OutFailReason))
 		{
