Updating SWizard to new styles - impacts Profile Wizard, Niagara New Asset wizards, and New Class Dialog. New Class Dialog updated to new style guidelines/using new widget types

diff --git a/Engine/Plugins/Editor/MobileLauncherProfileWizard/Source/MobileLauncherProfileWizard/Private/SProfileWizardUI.cpp b/Engine/Plugins/Editor/MobileLauncherProfileWizard/Source/MobileLauncherProfileWizard/Private/SProfileWizardUI.cpp
--- a/Engine/Plugins/Editor/MobileLauncherProfileWizard/Source/MobileLauncherProfileWizard/Private/SProfileWizardUI.cpp
+++ b/Engine/Plugins/Editor/MobileLauncherProfileWizard/Source/MobileLauncherProfileWizard/Private/SProfileWizardUI.cpp
@@ -70,11 +70,6 @@ void SProfileWizardUI::Construct(const FArguments& InArgs)
 			[
 				SAssignNew(MainWizard, SWizard)
 				.ShowPageList(false)
-				.ButtonStyle(FEditorStyle::Get(), "FlatButton.Default")
-				.CancelButtonStyle(FEditorStyle::Get(), "FlatButton.Default")
-				.FinishButtonStyle(FEditorStyle::Get(), "FlatButton.Success")
-				.ButtonTextStyle(FEditorStyle::Get(), "LargeText")
-				.ForegroundColor(FEditorStyle::Get().GetSlateColor("WhiteBrush"))
 				.CanFinish(this, &SProfileWizardUI::CanFinish)
 				.FinishButtonText(LOCTEXT("FinishButtonText", "Create Profile"))
 				.FinishButtonToolTip(LOCTEXT("FinishButtonToolTip", "Creates the launcher profiles for packaging a simple application and downloadable content."))
diff --git a/Engine/Plugins/FX/Niagara/Source/NiagaraEditor/Private/SNiagaraNewAssetDialog.cpp b/Engine/Plugins/FX/Niagara/Source/NiagaraEditor/Private/SNiagaraNewAssetDialog.cpp
--- a/Engine/Plugins/FX/Niagara/Source/NiagaraEditor/Private/SNiagaraNewAssetDialog.cpp
+++ b/Engine/Plugins/FX/Niagara/Source/NiagaraEditor/Private/SNiagaraNewAssetDialog.cpp
@@ -115,10 +115,6 @@ void SNiagaraNewAssetDialog::Construct(const FArguments& InArgs, FName InSaveCon
 			.OnFinished(this, &SNiagaraNewAssetDialog::OnOkButtonClicked)
 			.CanFinish(this, &SNiagaraNewAssetDialog::IsOkButtonEnabled)
 			.ShowPageList(false)
-			.ButtonStyle(FEditorStyle::Get(), "FlatButton.Default")
-			.CancelButtonStyle(FEditorStyle::Get(), "FlatButton.Default")
-			.FinishButtonStyle(FEditorStyle::Get(), "FlatButton.Success")
-			.ButtonTextStyle(FEditorStyle::Get(), "FlatButton.DefaultTextStyle")
 			+SWizard::Page()
 			.CanShow(true)
 			.OnEnter(this, &SNiagaraNewAssetDialog::ResetStage)
diff --git a/Engine/Source/Editor/ClassViewer/Private/SClassViewer.cpp b/Engine/Source/Editor/ClassViewer/Private/SClassViewer.cpp
--- a/Engine/Source/Editor/ClassViewer/Private/SClassViewer.cpp
+++ b/Engine/Source/Editor/ClassViewer/Private/SClassViewer.cpp
@@ -1590,6 +1590,7 @@ void SClassViewer::Construct(const FArguments& InArgs, const FClassViewerInitial
 					.Text(InitOptions.ViewerTitleString)
 				]
 			]
+
 			+SVerticalBox::Slot()
 			.AutoHeight()
 			[
@@ -1601,12 +1602,30 @@ void SClassViewer::Construct(const FArguments& InArgs, const FClassViewerInitial
 					FiltersWidget
 				]
 				+ SHorizontalBox::Slot()
-				.Padding(2.0f, 2.0f)
+				.Padding(2.0f, 2.0f, 6.0f, 2.0f)
 				[
 					SAssignNew(SearchBox, SSearchBox)
 					.OnTextChanged( this, &SClassViewer::OnFilterTextChanged )
 					.OnTextCommitted( this, &SClassViewer::OnFilterTextCommitted )
 				]
+				// View mode combo button
+				+ SHorizontalBox::Slot()
+				.AutoWidth()
+				.Padding(2.0f, 2.0f)
+				[
+					SAssignNew(ViewOptionsComboButton, SComboButton)
+					.ContentPadding(0)
+					.ForegroundColor(FSlateColor::UseForeground())
+					.ComboButtonStyle(FAppStyle::Get(), "SimpleComboButton")
+					.HasDownArrow(false)
+					.OnGetMenuContent(this, &SClassViewer::GetViewButtonContent)
+					.ButtonContent()
+					[
+						SNew(SImage)
+						.Image(FAppStyle::Get().GetBrush("Icons.Settings"))
+						.ColorAndOpacity(FSlateColor::UseForeground())
+					]
+				]
 			]
 
 			+SVerticalBox::Slot()
@@ -1661,47 +1680,11 @@ void SClassViewer::Construct(const FArguments& InArgs, const FClassViewerInitial
 			// Bottom panel
 			+ SVerticalBox::Slot()
 			.AutoHeight()
+			.VAlign(VAlign_Center)
+			.Padding(4.0f)
 			[
-				SNew(SHorizontalBox)
-
-				// Asset count
-				+ SHorizontalBox::Slot()
-				.FillWidth(1.f)
-				.VAlign(VAlign_Center)
-				.Padding(8, 0)
-				[
-					SNew(STextBlock)
-					.Text(this, &SClassViewer::GetClassCountText)
-				]
-
-				// View mode combo button
-				+ SHorizontalBox::Slot()
-				.AutoWidth()
-				[
-					SAssignNew(ViewOptionsComboButton, SComboButton)
-					.ContentPadding(0)
-					.ForegroundColor(this, &SClassViewer::GetViewButtonForegroundColor)
-					.ButtonStyle(FEditorStyle::Get(), "ToggleButton") // Use the tool bar item style for this button
-					.OnGetMenuContent(this, &SClassViewer::GetViewButtonContent)
-					.ButtonContent()
-					[
-						SNew(SHorizontalBox)
-							+ SHorizontalBox::Slot()
-						.AutoWidth()
-						.VAlign(VAlign_Center)
-						[
-							SNew(SImage).Image(FEditorStyle::GetBrush("GenericViewButton"))
-						]
-
-						+ SHorizontalBox::Slot()
-						.AutoWidth()
-						.Padding(2, 0, 0, 0)
-						.VAlign(VAlign_Center)
-						[
-							SNew(STextBlock).Text(LOCTEXT("ViewButton", "View Options"))
-						]
-					]
-				]
+				SNew(STextBlock)
+				.Text(this, &SClassViewer::GetClassCountText)
 			]
 		]
 	];
@@ -1926,14 +1909,6 @@ const int SClassViewer::GetNumItems() const
 	return NumClasses;
 }
 
-FSlateColor SClassViewer::GetViewButtonForegroundColor() const
-{
-	static const FName InvertedForegroundName("InvertedForeground");
-	static const FName DefaultForegroundName("DefaultForeground");
-
-	return ViewOptionsComboButton->IsHovered() ? FEditorStyle::GetSlateColor(InvertedForegroundName) : FEditorStyle::GetSlateColor(DefaultForegroundName);
-}
-
 TSharedRef<SWidget> SClassViewer::GetViewButtonContent()
 {
 	// Get all menu extenders for this context menu from the content browser module
diff --git a/Engine/Source/Editor/ClassViewer/Public/SClassViewer.h b/Engine/Source/Editor/ClassViewer/Public/SClassViewer.h
--- a/Engine/Source/Editor/ClassViewer/Public/SClassViewer.h
+++ b/Engine/Source/Editor/ClassViewer/Public/SClassViewer.h
@@ -177,8 +177,7 @@ private:
 	/** Expands all of the root nodes */
 	virtual void ExpandRootNodes();
 
-	/** Returns the foreground color for the view button */
-	FSlateColor GetViewButtonForegroundColor() const;
+
 
 	/** Handler for when the view combo button is clicked */
 	TSharedRef<SWidget> GetViewButtonContent();
diff --git a/Engine/Source/Editor/EditorStyle/Private/StarshipStyle.cpp b/Engine/Source/Editor/EditorStyle/Private/StarshipStyle.cpp
--- a/Engine/Source/Editor/EditorStyle/Private/StarshipStyle.cpp
+++ b/Engine/Source/Editor/EditorStyle/Private/StarshipStyle.cpp
@@ -1593,44 +1593,22 @@ void FStarshipEditorStyle::FStyle::SetupGeneralStyles()
 	// NewClassDialog
 	if (IncludeEditorSpecificStyles())
 	{
-		Set( "NewClassDialog.PageTitle", FTextBlockStyle(NormalText)
-			.SetFont( DEFAULT_FONT( "BoldCondensed", 28 ) )
-			.SetShadowOffset( FVector2D( 1,1 ) )
-			.SetShadowColorAndOpacity( FLinearColor(0,0,0,0.9f) )
-		);
-
-		Set( "NewClassDialog.SelectedParentClassLabel", FTextBlockStyle(NormalText)
-			.SetFont( DEFAULT_FONT( "Regular", 12 ) )
-			.SetShadowOffset( FVector2D( 0,1 ) )
-			.SetShadowColorAndOpacity( FLinearColor(0,0,0,0.9f) )
-		);
-
-		Set( "NewClassDialog.ErrorLabelFont", FTextBlockStyle(NormalText)
-			.SetFont( DEFAULT_FONT( "Regular", 10 ) )
-		 );
-	
-		Set( "NewClassDialog.ErrorLabelBorder", new FSlateColorBrush( FLinearColor(0.2f, 0.f, 0.f, 0.7f) ) );
 		Set( "NewClassDialog.ErrorLabelCloseButton", new IMAGE_BRUSH( "Icons/Cross_12x", Icon12x12 ) );
 
 		Set( "NewClassDialog.ParentClassListView.TableRow", FTableRowStyle()
 			.SetEvenRowBackgroundBrush( FSlateNoResource() )
-			.SetEvenRowBackgroundHoveredBrush( IMAGE_BRUSH( "Common/Selection", Icon8x8, FLinearColor(1.0f, 1.0f, 1.0f, 0.1f) ) )
+			.SetEvenRowBackgroundHoveredBrush(FSlateRoundedBoxBrush(FStyleColors::Panel, 4.0f))
 			.SetOddRowBackgroundBrush( FSlateNoResource() )
-			.SetOddRowBackgroundHoveredBrush( IMAGE_BRUSH( "Common/Selection", Icon8x8, FLinearColor(1.0f, 1.0f, 1.0f, 0.1f) ) )
-			.SetSelectorFocusedBrush( BORDER_BRUSH( "Common/Selector", FMargin(4.f/16.f), SelectorColor ) )
-			.SetActiveBrush( IMAGE_BRUSH( "Common/Selection", Icon8x8, SelectionColor ) )
-			.SetActiveHoveredBrush( IMAGE_BRUSH( "Common/Selection", Icon8x8, SelectionColor ) )
-			.SetInactiveBrush( IMAGE_BRUSH( "Common/Selection", Icon8x8, SelectionColor ) )
-			.SetInactiveHoveredBrush( IMAGE_BRUSH( "Common/Selection", Icon8x8, SelectionColor )  )
-			.SetTextColor( DefaultForeground )
-			.SetSelectedTextColor( InvertedForeground )
+			.SetOddRowBackgroundHoveredBrush(FSlateRoundedBoxBrush(FStyleColors::Panel, 4.0f))
+			.SetSelectorFocusedBrush(FSlateRoundedBoxBrush(FStyleColors::Transparent, 4.0f, FStyleColors::Select, 1.0f))
+			.SetActiveBrush(FSlateRoundedBoxBrush(FStyleColors::Transparent, 4.0f, FStyleColors::Select, 1.0f))
+			.SetActiveHoveredBrush(FSlateRoundedBoxBrush(FStyleColors::Transparent, 4.0f, FStyleColors::Select, 1.0f))
+			.SetInactiveBrush(FSlateRoundedBoxBrush(FStyleColors::Transparent, 4.0f, FStyleColors::SelectInactive, 1.0f))
+			.SetInactiveHoveredBrush(FSlateRoundedBoxBrush(FStyleColors::Transparent, 4.0f, FStyleColors::SelectInactive, 1.0f))
+			.SetTextColor( FStyleColors::Foreground )
+			.SetSelectedTextColor(FStyleColors::Foreground)
 			);
 
-		Set( "NewClassDialog.ParentClassItemTitle", FTextBlockStyle(NormalText)
-			.SetFont( DEFAULT_FONT( "Regular", 14 ) )
-			.SetShadowOffset( FVector2D( 0,1 ) )
-			.SetShadowColorAndOpacity( FLinearColor(0,0,0,0.9f) )
-		);
 	}
 
 	// Package Migration
diff --git a/Engine/Source/Editor/GameProjectGeneration/Private/SNewClassDialog.cpp b/Engine/Source/Editor/GameProjectGeneration/Private/SNewClassDialog.cpp
--- a/Engine/Source/Editor/GameProjectGeneration/Private/SNewClassDialog.cpp
+++ b/Engine/Source/Editor/GameProjectGeneration/Private/SNewClassDialog.cpp
@@ -40,6 +40,9 @@
 #include "FeaturedClasses.inl"
 #include "Subsystems/AssetEditorSubsystem.h"
 #include "Editor.h"
+#include "Styling/StyleColors.h"
+#include "Widgets/Input/SSegmentedControl.h"
+#include "ClassIconFinder.h"
 
 #define LOCTEXT_NAMESPACE "GameProjectGeneration"
 
@@ -274,13 +277,6 @@ void SNewClassDialog::Construct( const FArguments& InArgs )
 			[
 				SAssignNew( MainWizard, SWizard)
 				.ShowPageList(false)
-
-				.ButtonStyle(FEditorStyle::Get(), "FlatButton.Default")
-				.CancelButtonStyle(FEditorStyle::Get(), "FlatButton.Default")
-				.FinishButtonStyle(FEditorStyle::Get(), "FlatButton.Success")
-				.ButtonTextStyle(FEditorStyle::Get(), "LargeText")
-				.ForegroundColor(FEditorStyle::Get().GetSlateColor("WhiteBrush"))
-
 				.CanFinish(this, &SNewClassDialog::CanFinish)
 				.FinishButtonText( ClassDomain == EClassDomain::Native ? LOCTEXT("FinishButtonText_Native", "Create Class") : LOCTEXT("FinishButtonText_Blueprint", "Create Blueprint Class") )
 				.FinishButtonToolTip (
@@ -296,7 +292,7 @@ void SNewClassDialog::Construct( const FArguments& InArgs )
 					// Get IDE information
 					SNew(SBorder)
 					.Visibility( this, &SNewClassDialog::GetGlobalErrorLabelVisibility )
-					.BorderImage( FEditorStyle::GetBrush("NewClassDialog.ErrorLabelBorder") )
+					.BorderImage(FAppStyle::Get().GetBrush("RoundedError"))
 					.Padding(FMargin(0, 5))
 					.Content()
 					[
@@ -308,7 +304,7 @@ void SNewClassDialog::Construct( const FArguments& InArgs )
 						.AutoWidth()
 						[
 							SNew(SImage)
-							.Image(FEditorStyle::GetBrush("MessageLog.Warning"))
+							.Image(FAppStyle::Get().GetBrush("Icons.ErrorWithColor"))
 						]
 
 						+SHorizontalBox::Slot()
@@ -316,7 +312,6 @@ void SNewClassDialog::Construct( const FArguments& InArgs )
 						[
 							SNew(STextBlock)
 							.Text( this, &SNewClassDialog::GetGlobalErrorLabelText )
-							.TextStyle( FEditorStyle::Get(), "NewClassDialog.ErrorLabelFont" )
 						]
 
 						+SHorizontalBox::Slot()
@@ -342,16 +337,22 @@ void SNewClassDialog::Construct( const FArguments& InArgs )
 					.Padding(0)
 					[
 						SNew(STextBlock)
-						.TextStyle( FEditorStyle::Get(), "NewClassDialog.PageTitle" )
+						.Font(FAppStyle::Get().GetFontStyle("HeadingExtraSmall"))
 						.Text( LOCTEXT( "ParentClassTitle", "Choose Parent Class" ) )
+						.TransformPolicy(ETextTransformPolicy::ToUpper)
 					]
 
-					// Title spacer
-					+SVerticalBox::Slot()
+					+ SVerticalBox::Slot()
 					.AutoHeight()
-					.Padding(0, 2, 0, 8)
+					.HAlign(HAlign_Center)
 					[
-						SNew(SSeparator)
+						SNew(SSegmentedControl<bool>)
+						.OnValueChanged(this, &SNewClassDialog::OnFullClassTreeChanged)
+						.Value(this, &SNewClassDialog::IsFullClassTreeShown)
+						+SSegmentedControl<bool>::Slot(false)
+						.Text(LOCTEXT("CommonClasses", "Common Classes"))
+						+ SSegmentedControl<bool>::Slot(true)
+						.Text(LOCTEXT("AllClasses", "All Classes"))
 					]
 
 					// Page description and view options
@@ -359,34 +360,12 @@ void SNewClassDialog::Construct( const FArguments& InArgs )
 					.AutoHeight()
 					.Padding(0, 10)
 					[
-						SNew(SHorizontalBox)
-
-						+SHorizontalBox::Slot()
-						.FillWidth(1.f)
-						.VAlign(VAlign_Center)
-						[
-							SNew(STextBlock)
-							.Text(
-								ClassDomain == EClassDomain::Native ?
-								LOCTEXT("ChooseParentClassDescription_Native", "This will add a C++ header and source code file to your game project.") :
-								LOCTEXT("ChooseParentClassDescription_Blueprint", "This will add a new Blueprint class to your game project.")
-								)
-						]
-
-						// Full tree checkbox
-						+SHorizontalBox::Slot()
-						.AutoWidth()
-						.VAlign(VAlign_Center)
-						.Padding(4, 0, 0, 0)
-						[
-							SNew(SCheckBox)
-							.IsChecked( this, &SNewClassDialog::IsFullClassTreeChecked )
-							.OnCheckStateChanged( this, &SNewClassDialog::OnFullClassTreeChanged )
-							[
-								SNew(STextBlock)
-								.Text( LOCTEXT( "FullClassTree", "Show All Classes" ) )
-							]
-						]
+						SNew(STextBlock)
+						.Text(
+							ClassDomain == EClassDomain::Native ?
+							LOCTEXT("ChooseParentClassDescription_Native", "This will add a C++ header and source code file to your game project.") :
+							LOCTEXT("ChooseParentClassDescription_Blueprint", "This will add a new Blueprint class to your game project.")
+						)
 					]
 
 					// Add Code list
@@ -430,77 +409,63 @@ void SNewClassDialog::Construct( const FArguments& InArgs )
 					.Padding(30, 2)
 					.AutoHeight()
 					[
-						SNew(SHorizontalBox)
-						
+						SNew(SGridPanel)
+						.FillColumn(1, 1.0f)
+
 						// Class label
-						+SHorizontalBox::Slot()
-						.AutoWidth()
+						+ SGridPanel::Slot(0,0)
+						.VAlign(VAlign_Center)
+						.Padding(2.0f, 2.0f, 10.0f, 2.0f)
+						.HAlign(HAlign_Left)
 						[
-							SNew(SVerticalBox)
+							SNew(STextBlock)
+							.Text(LOCTEXT("ParentClassLabel", "Selected Class"))
+						]
 
-							+SVerticalBox::Slot()
-							.AutoHeight()
+						+ SGridPanel::Slot(0, 1)
+						.VAlign(VAlign_Center)
+						.Padding(2.0f, 2.0f, 10.0f, 2.0f)
+						.HAlign(HAlign_Left)
+						[
+							SNew(STextBlock)
+							.Visibility(ClassDomain == EClassDomain::Blueprint ? EVisibility::Collapsed : EVisibility::Visible)
+							.Text(LOCTEXT("ParentClassSourceLabel", "Selected Class Source"))
+						]
+
+						+ SGridPanel::Slot(1, 0)
+						.VAlign(VAlign_Center)
+						.Padding(2.0f)
+						.HAlign(HAlign_Left)
+						[
+							SNew(SHorizontalBox)
+
+							+ SHorizontalBox::Slot()
 							.VAlign(VAlign_Center)
-							.Padding(0, 0, 12, 0)
+							.AutoWidth()
 							[
 								SNew(STextBlock)
-								.TextStyle( FEditorStyle::Get(), "NewClassDialog.SelectedParentClassLabel" )
-								.Text( LOCTEXT( "ParentClassLabel", "Selected Class" ) )
+								.Text(this, &SNewClassDialog::GetSelectedParentClassName)
 							]
 
-							+SVerticalBox::Slot()
-							.AutoHeight()
+							+ SHorizontalBox::Slot()
 							.VAlign(VAlign_Center)
-							.Padding(0, 0, 12, 0)
+							.AutoWidth()
 							[
-								SNew(STextBlock)
-								.Visibility(ClassDomain == EClassDomain::Blueprint ? EVisibility::Collapsed : EVisibility::Visible)
-								.TextStyle(FEditorStyle::Get(), "NewClassDialog.SelectedParentClassLabel")
-								.Text(LOCTEXT("ParentClassSourceLabel", "Selected Class Source"))
+								DocWidget
 							]
 						]
 
-						// Class selection preview
-						+SHorizontalBox::Slot()
+						+ SGridPanel::Slot(1, 1)
+						.VAlign(VAlign_Center)
+						.Padding(2.0f)
+						.HAlign(HAlign_Left)
 						[
-							SNew(SVerticalBox)
-
-							+SVerticalBox::Slot()
-							.AutoHeight()
-							.VAlign(VAlign_Center)
-							.Padding(0, 0, 12, 0)
-							[
-								SNew(SHorizontalBox)
-
-								+SHorizontalBox::Slot()
-								.VAlign(VAlign_Center)
-								.AutoWidth()
-								[
-									SNew(STextBlock)
-									.Text( this, &SNewClassDialog::GetSelectedParentClassName )
-								]
-
-								+SHorizontalBox::Slot()
-								.VAlign(VAlign_Center)
-								.AutoWidth()
-								[
-									DocWidget
-								]
-							]
-
-							+ SVerticalBox::Slot()
-							.AutoHeight()
-							.VAlign(VAlign_Bottom)
-							.HAlign(HAlign_Left)
-							.Padding(0.0f, 0.0f, 0.0f, 0.0f)
-							[
-								SNew(SHyperlink)
-								.Style(FEditorStyle::Get(), "Common.GotoNativeCodeHyperlink")
-								.OnNavigate(this, &SNewClassDialog::OnEditCodeClicked)
-								.Text(this, &SNewClassDialog::GetSelectedParentClassFilename)
-								.ToolTipText(FText::Format(LOCTEXT("GoToCode_ToolTip", "Click to open this source file in {0}"), FSourceCodeNavigation::GetSelectedSourceCodeIDE()))
-								.Visibility(this, &SNewClassDialog::GetSourceHyperlinkVisibility)
-							]
+							SNew(SHyperlink)
+							.Style(FEditorStyle::Get(), "Common.GotoNativeCodeHyperlink")
+							.OnNavigate(this, &SNewClassDialog::OnEditCodeClicked)
+							.Text(this, &SNewClassDialog::GetSelectedParentClassFilename)
+							.ToolTipText(FText::Format(LOCTEXT("GoToCode_ToolTip", "Click to open this source file in {0}"), FSourceCodeNavigation::GetSelectedSourceCodeIDE()))
+							.Visibility(this, &SNewClassDialog::GetSourceHyperlinkVisibility)
 						]
 					]
 				]
@@ -517,16 +482,9 @@ void SNewClassDialog::Construct( const FArguments& InArgs )
 					.Padding(0)
 					[
 						SNew(STextBlock)
-						.TextStyle( FEditorStyle::Get(), "NewClassDialog.PageTitle" )
+						.Font(FAppStyle::Get().GetFontStyle("HeadingExtraSmall"))
 						.Text( this, &SNewClassDialog::GetNameClassTitle )
-					]
-
-					// Title spacer
-					+SVerticalBox::Slot()
-					.AutoHeight()
-					.Padding(0, 2, 0, 8)
-					[
-						SNew(SSeparator)
+						.TransformPolicy(ETextTransformPolicy::ToUpper)
 					]
 
 					+SVerticalBox::Slot()
@@ -560,7 +518,8 @@ void SNewClassDialog::Construct( const FArguments& InArgs )
 						.Padding(0, 5)
 						[
 							// Constant height, whether the label is visible or not
-							SNew(SBox).HeightOverride(20)
+							SNew(SBox)
+							.HeightOverride(20)
 							[
 								SNew(SBorder)
 								.Visibility( this, &SNewClassDialog::GetNameErrorLabelVisibility )
@@ -591,19 +550,43 @@ void SNewClassDialog::Construct( const FArguments& InArgs )
 								[
 									SNew(SGridPanel)
 									.FillColumn(1, 1.0f)
+									// Class type label
+
+									+ SGridPanel::Slot(0, 0)
+									.VAlign(VAlign_Center)
+									.Padding(0, 0, 12, 0)
+									[
+										SNew(STextBlock)
+										.Text(LOCTEXT("ClassTypeLabel", "Class Type"))
+									]
 
+									+SGridPanel::Slot(1,0)
+									.VAlign(VAlign_Center)
+									.HAlign(HAlign_Left)
+									.Padding(2.0f)
+									[
+										SNew(SSegmentedControl<GameProjectUtils::EClassLocation>)
+										.Visibility(ClassDomain == EClassDomain::Blueprint ? EVisibility::Collapsed : EVisibility::Visible)
+										.OnValueChanged(this, &SNewClassDialog::OnClassLocationChanged)
+										.Value(this, &SNewClassDialog::IsClassLocationActive)
+										+ SSegmentedControl<GameProjectUtils::EClassLocation>::Slot(GameProjectUtils::EClassLocation::Public)
+										.Text(LOCTEXT("Public", "Public"))
+										.ToolTip(LOCTEXT("ClassLocation_Public", "A public class can be included and used inside other modules in addition to the module it resides in"))
+										+ SSegmentedControl<GameProjectUtils::EClassLocation>::Slot(GameProjectUtils::EClassLocation::Private)
+										.Text(LOCTEXT("Private", "Private"))
+										.ToolTip(LOCTEXT("ClassLocation_Private", "A private class can only be included and used within the module it resides in"))
+									]
 									// Name label
-									+SGridPanel::Slot(0, 0)
+									+SGridPanel::Slot(0, 1)
 									.VAlign(VAlign_Center)
 									.Padding(0, 0, 12, 0)
 									[
 										SNew(STextBlock)
-										.TextStyle( FEditorStyle::Get(), "NewClassDialog.SelectedParentClassLabel" )
 										.Text( LOCTEXT( "NameLabel", "Name" ) )
 									]
 
 									// Name edit box
-									+SGridPanel::Slot(1, 0)
+									+SGridPanel::Slot(1, 1)
 									.Padding(0.0f, 3.0f)
 									.VAlign(VAlign_Center)
 									[
@@ -614,7 +597,7 @@ void SNewClassDialog::Construct( const FArguments& InArgs )
 											SNew(SHorizontalBox)
 
 											+SHorizontalBox::Slot()
-											.FillWidth(1.0f)
+											.FillWidth(.7f)
 											[
 												SAssignNew( ClassNameEditBox, SEditableTextBox)
 												.Text( this, &SNewClassDialog::OnGetClassNameText )
@@ -638,72 +621,20 @@ void SNewClassDialog::Construct( const FArguments& InArgs )
 													.Text( this, &SNewClassDialog::GetSelectedModuleComboText )
 												]
 											]
-
-											// Native C++ properties
-											+SHorizontalBox::Slot()
-											.AutoWidth()
-											.Padding(6.0f, 0.0f, 0.0f, 0.0f)
-											[
-												SNew(SHorizontalBox)
-												.Visibility(ClassDomain == EClassDomain::Blueprint ? EVisibility::Collapsed : EVisibility::Visible)
-
-												+SHorizontalBox::Slot()
-												.AutoWidth()
-												[
-													SNew(SCheckBox)
-													.Style(FEditorStyle::Get(), "Property.ToggleButton.Start")
-													.IsChecked(this, &SNewClassDialog::IsClassLocationActive, GameProjectUtils::EClassLocation::Public)
-													.OnCheckStateChanged(this, &SNewClassDialog::OnClassLocationChanged, GameProjectUtils::EClassLocation::Public)
-													.ToolTipText(LOCTEXT("ClassLocation_Public", "A public class can be included and used inside other modules in addition to the module it resides in"))
-													[
-														SNew(SBox)
-														.VAlign(VAlign_Center)
-														.HAlign(HAlign_Left)
-														.Padding(FMargin(4.0f, 0.0f, 3.0f, 0.0f))
-														[
-															SNew(STextBlock)
-															.Text(LOCTEXT("Public", "Public"))
-															.ColorAndOpacity(this, &SNewClassDialog::GetClassLocationTextColor, GameProjectUtils::EClassLocation::Public)
-														]
-													]
-												]
-
-												+SHorizontalBox::Slot()
-												.AutoWidth()
-												[
-													SNew(SCheckBox)
-													.Style(FEditorStyle::Get(), "Property.ToggleButton.End")
-													.IsChecked(this, &SNewClassDialog::IsClassLocationActive, GameProjectUtils::EClassLocation::Private)
-													.OnCheckStateChanged(this, &SNewClassDialog::OnClassLocationChanged, GameProjectUtils::EClassLocation::Private)
-													.ToolTipText(LOCTEXT("ClassLocation_Private", "A private class can only be included and used within the module it resides in"))
-													[
-														SNew(SBox)
-														.VAlign(VAlign_Center)
-														.HAlign(HAlign_Right)
-														.Padding(FMargin(3.0f, 0.0f, 4.0f, 0.0f))
-														[
-															SNew(STextBlock)
-															.Text(LOCTEXT("Private", "Private"))
-															.ColorAndOpacity(this, &SNewClassDialog::GetClassLocationTextColor, GameProjectUtils::EClassLocation::Private)
-														]
-													]
-												]
-											]
 										]
 									]
 
 									// Path label
-									+SGridPanel::Slot(0, 1)
+									+SGridPanel::Slot(0, 2)
 									.VAlign(ClassDomain == EClassDomain::Blueprint ? VAlign_Top : VAlign_Center)
 									.Padding(0, 0, 12, 0)
 									[
 										SNew(STextBlock)
-										.TextStyle( FEditorStyle::Get(), "NewClassDialog.SelectedParentClassLabel" )
 										.Text( LOCTEXT( "PathLabel", "Path" ) )
 									]
 
 									// Path edit box
-									+SGridPanel::Slot(1, 1)
+									+SGridPanel::Slot(1, 2)
 									.Padding(0.0f, 3.0f)
 									.VAlign(VAlign_Center)
 									[
@@ -748,26 +679,30 @@ void SNewClassDialog::Construct( const FArguments& InArgs )
 												[
 													SNew(SButton)
 													.VAlign(VAlign_Center)
+													.ButtonStyle(FAppStyle::Get(), "SimpleButton")
 													.OnClicked(this, &SNewClassDialog::HandleChooseFolderButtonClicked)
-													.Text( LOCTEXT( "BrowseButtonText", "Choose Folder" ) )
+													[
+														SNew(SImage)
+														.Image(FAppStyle::Get().GetBrush("Icons.FolderClosed"))
+														.ColorAndOpacity(FSlateColor::UseForeground())
+													]
 												]
 											]
 										]
 									]
 
 									// Header output label
-									+SGridPanel::Slot(0, 2)
+									+SGridPanel::Slot(0, 3)
 									.VAlign(VAlign_Center)
 									.Padding(0, 0, 12, 0)
 									[
 										SNew(STextBlock)
 										.Visibility(ClassDomain == EClassDomain::Blueprint ? EVisibility::Collapsed : EVisibility::Visible)
-										.TextStyle( FEditorStyle::Get(), "NewClassDialog.SelectedParentClassLabel" )
 										.Text( LOCTEXT( "HeaderFileLabel", "Header File" ) )
 									]
 
 									// Header output text
-									+SGridPanel::Slot(1, 2)
+									+SGridPanel::Slot(1, 3)
 									.Padding(0.0f, 3.0f)
 									.VAlign(VAlign_Center)
 									[
@@ -782,18 +717,17 @@ void SNewClassDialog::Construct( const FArguments& InArgs )
 									]
 
 									// Source output label
-									+SGridPanel::Slot(0, 3)
+									+SGridPanel::Slot(0, 4)
 									.VAlign(VAlign_Center)
 									.Padding(0, 0, 12, 0)
 									[
 										SNew(STextBlock)
 										.Visibility(ClassDomain == EClassDomain::Blueprint ? EVisibility::Collapsed : EVisibility::Visible)
-										.TextStyle( FEditorStyle::Get(), "NewClassDialog.SelectedParentClassLabel" )
 										.Text( LOCTEXT( "SourceFileLabel", "Source File" ) )
 									]
 
 									// Source output text
-									+SGridPanel::Slot(1, 3)
+									+SGridPanel::Slot(1, 4)
 									.Padding(0.0f, 3.0f)
 									.VAlign(VAlign_Center)
 									[
@@ -809,16 +743,6 @@ void SNewClassDialog::Construct( const FArguments& InArgs )
 								]
 							]
 						]
-
-						+SVerticalBox::Slot()
-						.AutoHeight()
-						.Padding(0.0f)
-						[
-							SNew(SBorder)
-							.Padding(FMargin(0.0f, 3.0f, 0.0f, 0.0f))
-							.BorderImage(FEditorStyle::GetBrush("DetailsView.CategoryBottom"))
-							.BorderBackgroundColor(FLinearColor(0.6f, 0.6f, 0.6f, 1.0f ))
-						]
 					]
 				]
 			]
@@ -884,51 +808,54 @@ TSharedRef<ITableRow> SNewClassDialog::MakeParentClassListViewWidget(TSharedPtr<
 	const FText ClassName = ParentClassItem->ParentClassInfo.GetClassName();
 	const FText ClassFullDescription = ParentClassItem->ParentClassInfo.GetClassDescription(/*bFullDescription*/true);
 	const FText ClassShortDescription = ParentClassItem->ParentClassInfo.GetClassDescription(/*bFullDescription*/false);
-	const FSlateBrush* const ClassBrush = ParentClassItem->ParentClassInfo.GetClassIcon();
 	const UClass* Class = ParentClassItem->ParentClassInfo.BaseClass;
+	const FSlateBrush* const ClassBrush = FClassIconFinder::FindThumbnailForClass(Class);
 
 	const int32 ItemHeight = 64;
-	const int32 DescriptionIndent = 32;
 	return
 		SNew( STableRow<TSharedPtr<FParentClassItem>>, OwnerTable )
+		.Padding(4)
 		.Style(FEditorStyle::Get(), "NewClassDialog.ParentClassListView.TableRow")
 		.ToolTip(IDocumentation::Get()->CreateToolTip(ClassFullDescription, nullptr, FEditorClassUtils::GetDocumentationPage(Class), FEditorClassUtils::GetDocumentationExcerpt(Class)))
 		[
-			SNew(SBox).HeightOverride(ItemHeight)
+			SNew(SBox)
+			.HeightOverride(ItemHeight)
 			[
-				SNew(SVerticalBox)
-
-				+SVerticalBox::Slot()
+				SNew(SHorizontalBox)
+				+ SHorizontalBox::Slot()
+				.AutoWidth()
+				.VAlign(VAlign_Center)
+				.HAlign(HAlign_Center)
 				.Padding(8)
-				.AutoHeight()
 				[
-					SNew(SHorizontalBox)
-
-					+SHorizontalBox::Slot()
-					.AutoWidth()
-					.VAlign(VAlign_Center)
-					.Padding(0, 0, 4, 0)
+					SNew(SBox)
+					.HeightOverride(ItemHeight / 2.0f)
+					.WidthOverride(ItemHeight / 2.0f)
 					[
 						SNew(SImage)
 						.Image(ClassBrush)
 					]
+				]
 
-					+SHorizontalBox::Slot()
-					.VAlign(VAlign_Center)
+				+ SHorizontalBox::Slot()
+				.FillWidth(1.0f)
+				.VAlign(VAlign_Center)
+				.Padding(4)
+				[
+					SNew(SVerticalBox)
+					+SVerticalBox::Slot()
 					[
 						SNew(STextBlock)
-						.TextStyle( FEditorStyle::Get(), "NewClassDialog.ParentClassItemTitle" )
+						.TextStyle(FAppStyle::Get(), "DialogButtonText")
 						.Text(ClassName)
 					]
-				]
 
-				+SVerticalBox::Slot()
-				.FillHeight(1.f)
-				.Padding(DescriptionIndent, 0, 0, 0)
-				[
-					SNew(STextBlock)
-					//.AutoWrapText(true)
-					.Text(ClassShortDescription)
+					+SVerticalBox::Slot()
+					[
+						SNew(STextBlock)
+						.Text(ClassShortDescription)
+						.AutoWrapText(true)
+					]
 				]
 			]
 		];
@@ -1022,14 +949,14 @@ void SNewClassDialog::OnAdvancedClassSelected(UClass* Class)
 	ParentClassInfo = FNewClassInfo(Class);
 }
 
-ECheckBoxState SNewClassDialog::IsFullClassTreeChecked() const
+bool SNewClassDialog::IsFullClassTreeShown() const
 {
-	return bShowFullClassTree ? ECheckBoxState::Checked : ECheckBoxState::Unchecked;
+	return bShowFullClassTree;
 }
 
-void SNewClassDialog::OnFullClassTreeChanged(ECheckBoxState NewCheckedState)
+void SNewClassDialog::OnFullClassTreeChanged(bool bInShowFullClassTree)
 {
-	bShowFullClassTree = (NewCheckedState == ECheckBoxState::Checked);
+	bShowFullClassTree = bInShowFullClassTree;
 }
 
 EVisibility SNewClassDialog::GetBasicParentClassVisibility() const
@@ -1435,64 +1362,61 @@ FSlateColor SNewClassDialog::GetClassLocationTextColor(GameProjectUtils::EClassL
 	return (ClassLocation == InLocation) ? FSlateColor(FLinearColor(0, 0, 0)) : FSlateColor(FLinearColor(0.72f, 0.72f, 0.72f, 1.f));
 }
 
-ECheckBoxState SNewClassDialog::IsClassLocationActive(GameProjectUtils::EClassLocation InLocation) const
+GameProjectUtils::EClassLocation SNewClassDialog::IsClassLocationActive() const
 {
-	return (ClassLocation == InLocation) ? ECheckBoxState::Checked : ECheckBoxState::Unchecked;
+	return ClassLocation;
 }
 
-void SNewClassDialog::OnClassLocationChanged(ECheckBoxState InCheckedState, GameProjectUtils::EClassLocation InLocation)
+void SNewClassDialog::OnClassLocationChanged(GameProjectUtils::EClassLocation InLocation)
 {
-	if(InCheckedState == ECheckBoxState::Checked)
-	{
-		const FString AbsoluteClassPath = FPaths::ConvertRelativePathToFull(NewClassPath) / ""; // Ensure trailing /
+	const FString AbsoluteClassPath = FPaths::ConvertRelativePathToFull(NewClassPath) / ""; // Ensure trailing /
 
-		GameProjectUtils::EClassLocation TmpClassLocation = GameProjectUtils::EClassLocation::UserDefined;
-		GameProjectUtils::GetClassLocation(AbsoluteClassPath, *SelectedModuleInfo, TmpClassLocation);
+	GameProjectUtils::EClassLocation TmpClassLocation = GameProjectUtils::EClassLocation::UserDefined;
+	GameProjectUtils::GetClassLocation(AbsoluteClassPath, *SelectedModuleInfo, TmpClassLocation);
 
-		const FString RootPath = SelectedModuleInfo->ModuleSourcePath;
-		const FString PublicPath = RootPath / "Public" / "";		// Ensure trailing /
-		const FString PrivatePath = RootPath / "Private" / "";		// Ensure trailing /
+	const FString RootPath = SelectedModuleInfo->ModuleSourcePath;
+	const FString PublicPath = RootPath / "Public" / "";		// Ensure trailing /
+	const FString PrivatePath = RootPath / "Private" / "";		// Ensure trailing /
 
-		// Update the class path to be rooted to the Public or Private folder based on InVisibility
-		switch(InLocation)
+	// Update the class path to be rooted to the Public or Private folder based on InVisibility
+	switch (InLocation)
+	{
+	case GameProjectUtils::EClassLocation::Public:
+		if (AbsoluteClassPath.StartsWith(PrivatePath))
 		{
-		case GameProjectUtils::EClassLocation::Public:
-			if(AbsoluteClassPath.StartsWith(PrivatePath))
-			{
-				NewClassPath = AbsoluteClassPath.Replace(*PrivatePath, *PublicPath);
-			}
-			else if(AbsoluteClassPath.StartsWith(RootPath))
-			{
-				NewClassPath = AbsoluteClassPath.Replace(*RootPath, *PublicPath);
-			}
-			else
-			{
-				NewClassPath = PublicPath;
-			}
-			break;
-
-		case GameProjectUtils::EClassLocation::Private:
-			if(AbsoluteClassPath.StartsWith(PublicPath))
-			{
-				NewClassPath = AbsoluteClassPath.Replace(*PublicPath, *PrivatePath);
-			}
-			else if(AbsoluteClassPath.StartsWith(RootPath))
-			{
-				NewClassPath = AbsoluteClassPath.Replace(*RootPath, *PrivatePath);
-			}
-			else
-			{
-				NewClassPath = PrivatePath;
-			}
-			break;
+			NewClassPath = AbsoluteClassPath.Replace(*PrivatePath, *PublicPath);
+		}
+		else if (AbsoluteClassPath.StartsWith(RootPath))
+		{
+			NewClassPath = AbsoluteClassPath.Replace(*RootPath, *PublicPath);
+		}
+		else
+		{
+			NewClassPath = PublicPath;
+		}
+		break;
 
-		default:
-			break;
+	case GameProjectUtils::EClassLocation::Private:
+		if (AbsoluteClassPath.StartsWith(PublicPath))
+		{
+			NewClassPath = AbsoluteClassPath.Replace(*PublicPath, *PrivatePath);
+		}
+		else if (AbsoluteClassPath.StartsWith(RootPath))
+		{
+			NewClassPath = AbsoluteClassPath.Replace(*RootPath, *PrivatePath);
+		}
+		else
+		{
+			NewClassPath = PrivatePath;
 		}
+		break;
 
-		// Will update ClassVisibility correctly
-		UpdateInputValidity();
+	default:
+		break;
 	}
+
+	// Will update ClassVisibility correctly
+	UpdateInputValidity();
 }
 
 void SNewClassDialog::UpdateInputValidity()
diff --git a/Engine/Source/Editor/GameProjectGeneration/Private/SNewClassDialog.h b/Engine/Source/Editor/GameProjectGeneration/Private/SNewClassDialog.h
--- a/Engine/Source/Editor/GameProjectGeneration/Private/SNewClassDialog.h
+++ b/Engine/Source/Editor/GameProjectGeneration/Private/SNewClassDialog.h
@@ -103,11 +103,11 @@ private:
 	/** Handler for when a class was picked in the full class tree */
 	void OnAdvancedClassSelected(UClass* Class);
 
-	/** Gets the check box state for the full class list */
-	ECheckBoxState IsFullClassTreeChecked() const;
+	/** Gets the boolean for the full class list */
+	bool IsFullClassTreeShown() const;
 
-	/** Gets the check box state for the full class list */
-	void OnFullClassTreeChanged(ECheckBoxState NewCheckedState);
+	/** Sets the boolean for the full class list */
+	void OnFullClassTreeChanged(bool bInShowFullClassTree);
 
 	/** Gets the visibility of the basic class list */
 	EVisibility GetBasicParentClassVisibility() const;
@@ -184,10 +184,10 @@ private:
 	FSlateColor GetClassLocationTextColor(GameProjectUtils::EClassLocation InLocation) const;
 
 	/** Checks to see if the given class location is active based on the current value of NewClassPath */
-	ECheckBoxState IsClassLocationActive(GameProjectUtils::EClassLocation InLocation) const;
+	GameProjectUtils::EClassLocation IsClassLocationActive() const;
 
 	/** Update the value of NewClassPath so that it uses the given class location */
-	void OnClassLocationChanged(ECheckBoxState InCheckedState, GameProjectUtils::EClassLocation InLocation);
+	void OnClassLocationChanged(GameProjectUtils::EClassLocation InLocation);
 
 	/** Checks the current class name/path for validity and updates cached values accordingly */
 	void UpdateInputValidity();
diff --git a/Engine/Source/Runtime/AppFramework/Private/Widgets/Workflow/SWizard.cpp b/Engine/Source/Runtime/AppFramework/Private/Widgets/Workflow/SWizard.cpp
--- a/Engine/Source/Runtime/AppFramework/Private/Widgets/Workflow/SWizard.cpp
+++ b/Engine/Source/Runtime/AppFramework/Private/Widgets/Workflow/SWizard.cpp
@@ -38,152 +38,154 @@ void SWizard::Construct( const FArguments& InArgs )
 
 	ChildSlot
 	[
-		SNew(SVerticalBox)
-
-		+ SVerticalBox::Slot()
-		.AutoHeight()
-		[
-			SAssignNew(BreadcrumbTrail, SBreadcrumbTrail<int32>)
-			.OnCrumbClicked(this, &SWizard::HandleBreadcrumbClicked)
-			.Visibility(InArgs._ShowBreadcrumbs ? EVisibility::Visible : EVisibility::Collapsed)
-		]
-		+ SVerticalBox::Slot()
-		.AutoHeight()
-		.Padding(InArgs._ShowBreadcrumbs ? FMargin(0, 10, 0, 0) : FMargin(0))
-		[
-			SNew(STextBlock)
-			.TextStyle(InArgs._PageTitleTextStyle)
-			.Text(this, &SWizard::HandleGetPageTitle)
-			.Visibility(InArgs._ShowPageTitle ? EVisibility::Visible : EVisibility::Collapsed)
-		]
-		+ SVerticalBox::Slot()
-		.FillHeight(1.0)
+		SNew(SBorder)
+		.BorderImage(FAppStyle::Get().GetBrush("Brushes.Panel"))
+		.Padding(10.0f)
 		[
-			SNew(SHorizontalBox)
+			SNew(SVerticalBox)
 
-			+ SHorizontalBox::Slot()
-			.AutoWidth()
-			.Padding(0.0f, 0.0f, 20.0f, 0.0f)
+			+ SVerticalBox::Slot()
+			.AutoHeight()
 			[
-				SAssignNew(PageListBox, SVerticalBox)
-				.Visibility(InArgs._ShowPageList ? EVisibility::Visible : EVisibility::Collapsed)
+				SAssignNew(BreadcrumbTrail, SBreadcrumbTrail<int32>)
+				.OnCrumbClicked(this, &SWizard::HandleBreadcrumbClicked)
+				.Visibility(InArgs._ShowBreadcrumbs ? EVisibility::Visible : EVisibility::Collapsed)
 			]
-
-			+ SHorizontalBox::Slot()
-			.FillWidth(1.0f)
+			+ SVerticalBox::Slot()
+			.AutoHeight()
+			.Padding(InArgs._ShowBreadcrumbs ? FMargin(0, 10, 0, 0) : FMargin(0))
 			[
-				// widget switcher
-				SAssignNew(WidgetSwitcher, SWidgetSwitcher)
+				SNew(STextBlock)
+				.TextStyle(InArgs._PageTitleTextStyle)
+				.Text(this, &SWizard::HandleGetPageTitle)
+				.Visibility(InArgs._ShowPageTitle ? EVisibility::Visible : EVisibility::Collapsed)
 			]
-		]
-
-		+ SVerticalBox::Slot()
-		.AutoHeight()
-		.Padding(0)
-		[
-			InArgs._PageFooter.Widget
-		]
-
-		+ SVerticalBox::Slot()
-		.AutoHeight()
-		.HAlign(HAlign_Right)
-		.Padding(0.0f, 10.0f, 0.0f, 0.0f)
-		[
-			SAssignNew(ButtonGrid, SUniformGridPanel)
-			.SlotPadding(FCoreStyle::Get().GetMargin("StandardDialog.SlotPadding"))
-			.MinDesiredSlotWidth(FCoreStyle::Get().GetFloat("StandardDialog.MinDesiredSlotWidth"))
-			.MinDesiredSlotHeight(FCoreStyle::Get().GetFloat("StandardDialog.MinDesiredSlotHeight"))
-			
-			+ SUniformGridPanel::Slot(0, 0)
+			+ SVerticalBox::Slot()
+			.FillHeight(1.0)
 			[
-				// 'Prev' button
-				SNew(SButton)
-				.ButtonStyle(InArgs._ButtonStyle)
-				.TextStyle(InArgs._ButtonTextStyle)
-				.ForegroundColor(InArgs._ForegroundColor)
-				.HAlign(HAlign_Center)
-				.VAlign(VAlign_Center)
-				.ContentPadding(FCoreStyle::Get().GetMargin("StandardDialog.ContentPadding"))
-				.IsEnabled(this, &SWizard::HandlePrevButtonIsEnabled)
-				.OnClicked(this, &SWizard::HandlePrevButtonClicked)
-				.Visibility(this, &SWizard::HandlePrevButtonVisibility)
-				.ToolTipText(LOCTEXT("PrevButtonTooltip", "Go back to the previous step"))
-				[
-					SNew(SHorizontalBox)
+				SNew(SHorizontalBox)
 
-					+ SHorizontalBox::Slot()
-					.Padding(2.0f, 0.0f)
-					.AutoWidth()
-					.VAlign(VAlign_Center)
-					[
-						SNew(SImage)
-						.Image(FCoreStyle::Get().GetBrush("Wizard.BackIcon"))
-						.ColorAndOpacity(FLinearColor(0.05f, 0.05f, 0.05f))
-					]
+				+ SHorizontalBox::Slot()
+				.AutoWidth()
+				.Padding(0.0f, 0.0f, 20.0f, 0.0f)
+				[
+					SAssignNew(PageListBox, SVerticalBox)
+					.Visibility(InArgs._ShowPageList ? EVisibility::Visible : EVisibility::Collapsed)
+				]
 
-					+ SHorizontalBox::Slot()
-					.AutoWidth()
-					.VAlign(VAlign_Center)
-					[
-						SNew(STextBlock)
-						.TextStyle(InArgs._ButtonTextStyle)
-						.Text(LOCTEXT("PrevButtonLabel", "Back"))
-					]
+				+ SHorizontalBox::Slot()
+				.FillWidth(1.0f)
+				[
+					// widget switcher
+					SAssignNew(WidgetSwitcher, SWidgetSwitcher)
 				]
 			]
 
-			+ SUniformGridPanel::Slot(1, 0)
+			+ SVerticalBox::Slot()
+			.AutoHeight()
+			.Padding(0)
 			[
-				// 'Next' button
-				SNew(SButton)
-				.ButtonStyle(InArgs._FinishButtonStyle)
-				.TextStyle(InArgs._ButtonTextStyle)
-				.ForegroundColor(InArgs._ForegroundColor)
-				.HAlign(HAlign_Center)
-				.VAlign(VAlign_Center)
-				.ContentPadding(FCoreStyle::Get().GetMargin("StandardDialog.ContentPadding"))
-				.IsEnabled(this, &SWizard::HandleNextButtonIsEnabled)
-				.OnClicked(this, &SWizard::HandleNextButtonClicked)
-				.Visibility(this, &SWizard::HandleNextButtonVisibility)
-				.ToolTipText(LOCTEXT("NextButtonTooltip", "Go to the next step"))
-				[
-					SNew(SHorizontalBox)
+				InArgs._PageFooter.Widget
+			]
 
-					+SHorizontalBox::Slot()
-					.AutoWidth()
+			+ SVerticalBox::Slot()
+			.AutoHeight()
+			.HAlign(HAlign_Right)
+			.Padding(0.0f, 10.0f, 0.0f, 0.0f)
+			[
+				SAssignNew(ButtonGrid, SUniformGridPanel)
+				.SlotPadding(FCoreStyle::Get().GetMargin("StandardDialog.SlotPadding"))
+				.MinDesiredSlotWidth(FCoreStyle::Get().GetFloat("StandardDialog.MinDesiredSlotWidth"))
+				.MinDesiredSlotHeight(FCoreStyle::Get().GetFloat("StandardDialog.MinDesiredSlotHeight"))
+			
+				+ SUniformGridPanel::Slot(0, 0)
+				[
+					// 'Prev' button
+					SNew(SButton)
+					.ButtonStyle(InArgs._ButtonStyle)
+					.TextStyle(InArgs._ButtonTextStyle)
+					.HAlign(HAlign_Center)
 					.VAlign(VAlign_Center)
+					.ContentPadding(FCoreStyle::Get().GetMargin("StandardDialog.ContentPadding"))
+					.IsEnabled(this, &SWizard::HandlePrevButtonIsEnabled)
+					.OnClicked(this, &SWizard::HandlePrevButtonClicked)
+					.Visibility(this, &SWizard::HandlePrevButtonVisibility)
+					.ToolTipText(LOCTEXT("PrevButtonTooltip", "Go back to the previous step"))
 					[
-						SNew(STextBlock)
-						.TextStyle(InArgs._ButtonTextStyle)
-						.Text(LOCTEXT("NextButtonLabel", "Next"))
+						SNew(SHorizontalBox)
+
+						+ SHorizontalBox::Slot()
+						.Padding(2.0f, 0.0f)
+						.AutoWidth()
+						.VAlign(VAlign_Center)
+						[
+							SNew(SImage)
+							.Image(FCoreStyle::Get().GetBrush("Wizard.BackIcon"))
+							.ColorAndOpacity(FLinearColor(0.05f, 0.05f, 0.05f))
+						]
+
+						+ SHorizontalBox::Slot()
+						.AutoWidth()
+						.VAlign(VAlign_Center)
+						[
+							SNew(STextBlock)
+							.TextStyle(InArgs._ButtonTextStyle)
+							.Text(LOCTEXT("PrevButtonLabel", "Back"))
+						]
 					]
+				]
 
-					+SHorizontalBox::Slot()
-					.Padding(2.0f, 0.0f)
-					.AutoWidth()
+				+ SUniformGridPanel::Slot(1, 0)
+				[
+					// 'Next' button
+					SNew(SButton)
+					.ButtonStyle(InArgs._FinishButtonStyle)
+					.TextStyle(InArgs._ButtonTextStyle)
+					.HAlign(HAlign_Center)
 					.VAlign(VAlign_Center)
+					.ContentPadding(FCoreStyle::Get().GetMargin("StandardDialog.ContentPadding"))
+					.IsEnabled(this, &SWizard::HandleNextButtonIsEnabled)
+					.OnClicked(this, &SWizard::HandleNextButtonClicked)
+					.Visibility(this, &SWizard::HandleNextButtonVisibility)
+					.ToolTipText(LOCTEXT("NextButtonTooltip", "Go to the next step"))
 					[
-						SNew(SImage)
-						.Image(FCoreStyle::Get().GetBrush("Wizard.NextIcon"))
-						.ColorAndOpacity(FLinearColor(0.05f, 0.05f, 0.05f))
+						SNew(SHorizontalBox)
+
+						+SHorizontalBox::Slot()
+						.AutoWidth()
+						.VAlign(VAlign_Center)
+						[
+							SNew(STextBlock)
+							.TextStyle(InArgs._ButtonTextStyle)
+							.Text(LOCTEXT("NextButtonLabel", "Next"))
+						]
+
+						+SHorizontalBox::Slot()
+						.Padding(2.0f, 0.0f)
+						.AutoWidth()
+						.VAlign(VAlign_Center)
+						[
+							SNew(SImage)
+							.Image(FCoreStyle::Get().GetBrush("Wizard.NextIcon"))
+							.ColorAndOpacity(FLinearColor(0.05f, 0.05f, 0.05f))
+						]
 					]
 				]
-			]
 
-			+ SUniformGridPanel::Slot(2,0)
-			[
-				// 'Finish' button
-				SNew(SButton)
-				.ButtonStyle(InArgs._FinishButtonStyle)
-				.TextStyle(InArgs._ButtonTextStyle)
-				.ForegroundColor(InArgs._ForegroundColor)
-				.HAlign(HAlign_Center)
-				.VAlign(VAlign_Center)
-				.ContentPadding(FCoreStyle::Get().GetMargin("StandardDialog.ContentPadding"))
-				.IsEnabled(InArgs._CanFinish)
-				.OnClicked(this, &SWizard::HandleFinishButtonClicked)
-				.ToolTipText(InArgs._FinishButtonToolTip)
-				.Text(InArgs._FinishButtonText)
+				+ SUniformGridPanel::Slot(2,0)
+				[
+					// 'Finish' button
+					SNew(SButton)
+					.ButtonStyle(InArgs._FinishButtonStyle)
+					.TextStyle(InArgs._ButtonTextStyle)
+					.HAlign(HAlign_Center)
+					.VAlign(VAlign_Center)
+					.ContentPadding(FCoreStyle::Get().GetMargin("StandardDialog.ContentPadding"))
+					.IsEnabled(InArgs._CanFinish)
+					.OnClicked(this, &SWizard::HandleFinishButtonClicked)
+					.ToolTipText(InArgs._FinishButtonToolTip)
+					.Text(InArgs._FinishButtonText)
+				]
 			]
 		]
 	];
@@ -196,7 +198,6 @@ void SWizard::Construct( const FArguments& InArgs )
 			SNew(SButton)
 			.ButtonStyle(InArgs._CancelButtonStyle)
 			.TextStyle(InArgs._ButtonTextStyle)
-			.ForegroundColor(InArgs._ForegroundColor)
 			.HAlign(HAlign_Center)
 			.VAlign(VAlign_Center)
 			.ContentPadding(FCoreStyle::Get().GetMargin("StandardDialog.ContentPadding"))
diff --git a/Engine/Source/Runtime/AppFramework/Public/Widgets/Workflow/SWizard.h b/Engine/Source/Runtime/AppFramework/Public/Widgets/Workflow/SWizard.h
--- a/Engine/Source/Runtime/AppFramework/Public/Widgets/Workflow/SWizard.h
+++ b/Engine/Source/Runtime/AppFramework/Public/Widgets/Workflow/SWizard.h
@@ -17,6 +17,7 @@
 #include "Styling/SlateTypes.h"
 #include "Styling/CoreStyle.h"
 #include "Framework/SlateDelegates.h"
+#include "Styling/AppStyle.h"
 
 /**
  * Delegate type for getting the next page to show.
@@ -165,11 +166,10 @@ public:
 public:
 
 	SLATE_BEGIN_ARGS(SWizard)
-		: _ButtonStyle(&FCoreStyle::Get().GetWidgetStyle<FButtonStyle>("Button"))
-		, _CancelButtonStyle(&FCoreStyle::Get().GetWidgetStyle<FButtonStyle>("Button"))
-		, _FinishButtonStyle(&FCoreStyle::Get().GetWidgetStyle<FButtonStyle>("Button"))
-		, _ButtonTextStyle(&FCoreStyle::Get().GetWidgetStyle< FTextBlockStyle >("NormalText"))
-		, _ForegroundColor(FCoreStyle::Get().GetSlateColor("InvertedForeground"))
+		: _ButtonStyle(&FAppStyle::Get().GetWidgetStyle<FButtonStyle>("SecondaryButton"))
+		, _CancelButtonStyle(&FAppStyle::Get().GetWidgetStyle<FButtonStyle>("SecondaryButton"))
+		, _FinishButtonStyle(&FAppStyle::Get().GetWidgetStyle<FButtonStyle>("PrimaryButton"))
+		, _ButtonTextStyle(&FAppStyle::Get().GetWidgetStyle<FTextBlockStyle>("DialogButtonText"))
 		, _CanFinish(true)
 		, _FinishButtonText(NSLOCTEXT("SWizard", "DefaultFinishButtonText", "Finish"))
 		, _FinishButtonToolTip(NSLOCTEXT("SWizard", "DefaultFinishButtonTooltip", "Finish the wizard"))
@@ -199,9 +199,6 @@ public:
 		/** The text style used by the buttons. */
 		SLATE_STYLE_ARGUMENT(FTextBlockStyle, ButtonTextStyle)
 
-		/** The text style used by the buttons. */
-		SLATE_ARGUMENT(FSlateColor, ForegroundColor)
-
 		/** Holds a flag indicating whether the 'Finish' button is enabled. */
 		SLATE_ATTRIBUTE(bool, CanFinish)
 
