- Remove CoreOnline ini config from project generation
- All shared ptr/ref are being changed to default ESPMode::ThreadSafe and ESPMode::Fast is being deprecated, so this migration path is not needed.

diff --git a/Engine/Source/Editor/GameProjectGeneration/Private/GameProjectUtils.cpp b/Engine/Source/Editor/GameProjectGeneration/Private/GameProjectUtils.cpp
--- a/Engine/Source/Editor/GameProjectGeneration/Private/GameProjectUtils.cpp
+++ b/Engine/Source/Editor/GameProjectGeneration/Private/GameProjectUtils.cpp
@@ -2171,10 +2171,6 @@ bool GameProjectUtils::GenerateConfigFiles(const FProjectInformation& InProjectI
 			}
 		}
 
-		FileContents += LINE_TERMINATOR;
-		FileContents += TEXT("[CoreOnline]") LINE_TERMINATOR;
-		FileContents += TEXT("bEnableThreadSafeUniqueNetIds=True") LINE_TERMINATOR;
-
 		if (WriteOutputFile(DefaultEngineIniFilename, FileContents, OutFailReason))
 		{
 			OutCreatedFiles.Add(DefaultEngineIniFilename);
