Disable World Partition by default for new projects.

diff --git a/Engine/Source/Editor/GameProjectGeneration/Private/GameProjectUtils.cpp b/Engine/Source/Editor/GameProjectGeneration/Private/GameProjectUtils.cpp
--- a/Engine/Source/Editor/GameProjectGeneration/Private/GameProjectUtils.cpp
+++ b/Engine/Source/Editor/GameProjectGeneration/Private/GameProjectUtils.cpp
@@ -293,7 +293,7 @@ namespace
 			ConfigValues.Emplace(TEXT("DefaultEngine.ini"),
 				TEXT("/Script/WorldPartitionEditor.WorldPartitionEditorSettings"),
 				TEXT("bEnableWorldPartition"),
-				TEXT("True"),
+				TEXT("False"),
 				true /* ShouldReplaceExistingValue */);
 
 			ConfigValues.Emplace(TEXT("DefaultEngine.ini"),
