Changing the warning to use SWarningOrErrorBox

diff --git a/Engine/Source/Editor/GameProjectGeneration/Private/SNewClassDialog.cpp b/Engine/Source/Editor/GameProjectGeneration/Private/SNewClassDialog.cpp
--- a/Engine/Source/Editor/GameProjectGeneration/Private/SNewClassDialog.cpp
+++ b/Engine/Source/Editor/GameProjectGeneration/Private/SNewClassDialog.cpp
@@ -43,6 +43,7 @@
 #include "Styling/StyleColors.h"
 #include "Widgets/Input/SSegmentedControl.h"
 #include "ClassIconFinder.h"
+#include "SWarningOrErrorBox.h"
 
 #define LOCTEXT_NAMESPACE "GameProjectGeneration"
 
@@ -517,20 +518,10 @@ void SNewClassDialog::Construct( const FArguments& InArgs )
 						.AutoHeight()
 						.Padding(0, 5)
 						[
-							// Constant height, whether the label is visible or not
-							SNew(SBox)
-							.HeightOverride(20)
-							[
-								SNew(SBorder)
-								.Visibility( this, &SNewClassDialog::GetNameErrorLabelVisibility )
-								.BorderImage( FEditorStyle::GetBrush("NewClassDialog.ErrorLabelBorder") )
-								.Content()
-								[
-									SNew(STextBlock)
-									.Text( this, &SNewClassDialog::GetNameErrorLabelText )
-									.TextStyle( FEditorStyle::Get(), "NewClassDialog.ErrorLabelFont" )
-								]
-							]
+							SNew(SWarningOrErrorBox)
+							.MessageStyle(EMessageStyle::Error)
+							.Visibility(this, &SNewClassDialog::GetNameErrorLabelVisibility)
+							.Message(this, &SNewClassDialog::GetNameErrorLabelText)
 						]
 
 						// Properties
