remove reference to deleted platform.

#rnx

diff --git a/Engine/Source/Editor/GameProjectGeneration/Private/GameProjectUtils.cpp b/Engine/Source/Editor/GameProjectGeneration/Private/GameProjectUtils.cpp
--- a/Engine/Source/Editor/GameProjectGeneration/Private/GameProjectUtils.cpp
+++ b/Engine/Source/Editor/GameProjectGeneration/Private/GameProjectUtils.cpp
@@ -153,12 +153,6 @@ namespace
 		WindowsSettings.Settings.NumBuffers = 7;
 		DefaultProjectSettings.Add(TEXT("Windows"), WindowsSettings);
 
-		FAudioDefaultPlatformSettings XBoxSettings(TEXT("/Script/XboxOnePlatformEditor.XboxOneTargetSettings"));
-		XBoxSettings.Settings.CallbackBufferFrameSize = 256;
-		XBoxSettings.Settings.NumBuffers = 7;
-		DefaultProjectSettings.Add(TEXT("XboxOne"), XBoxSettings);
-
-
 		return MoveTemp(DefaultProjectSettings);
 	}
 
@@ -712,12 +706,6 @@ bool GameProjectUtils::IsValidProjectFileForCreation(const FString& ProjectFile,
 		return false;
 	}
 
-	if (NameContainsUnderscoreAndXB1Installed(BaseProjectFile))
-	{
-		OutFailReason = LOCTEXT( "ProjectNameContainsIllegalCharactersOnXB1", "Project names may not contain an underscore when the Xbox One XDK is installed." );
-		return false;
-	}
-
 	if ( !FPaths::ValidatePath(FPaths::GetPath(ProjectFile), &OutFailReason) )
 	{
 		return false;
@@ -802,12 +790,6 @@ bool GameProjectUtils::OpenProject(const FString& ProjectFile, FText& OutFailRea
 		return false;
 	}
 
-	if (NameContainsUnderscoreAndXB1Installed(BaseProjectFile))
-	{
-		OutFailReason = LOCTEXT( "ProjectNameContainsIllegalCharactersOnXB1", "Project names may not contain an underscore when the Xbox One XDK is installed." );
-		return false;
-	}
-
 	if ( !FPaths::ValidatePath(FPaths::GetPath(ProjectFile), &OutFailReason) )
 	{
 		return false;
@@ -2004,34 +1986,6 @@ bool GameProjectUtils::NameContainsOnlyLegalCharacters(const FString& TestName,
 	return !bContainsIllegalCharacters;
 }
 
-bool GameProjectUtils::NameContainsUnderscoreAndXB1Installed(const FString& TestName)
-{
-	// disabled for now so people with the SDK installed can use the editor
-	return false;
-
-	bool bContainsIllegalCharacters = false;
-
-	// Only allow alphanumeric characters in the project name
-	for ( int32 CharIdx = 0 ; CharIdx < TestName.Len() ; ++CharIdx )
-	{
-		const FString& Char = TestName.Mid( CharIdx, 1 );
-		if ( Char == TEXT("_") )
-		{
-			const ITargetPlatform* Platform = GetTargetPlatformManager()->FindTargetPlatform(TEXT("XboxOne"));
-			if (Platform)
-			{
-				FString NotInstalledDocLink;
-				if (Platform->IsSdkInstalled(true, NotInstalledDocLink))
-				{
-					bContainsIllegalCharacters = true;
-				}
-			}
-		}
-	}
-
-	return bContainsIllegalCharacters;
-}
-
 bool GameProjectUtils::ProjectFileExists(const FString& ProjectFile)
 {
 	return FPlatformFileManager::Get().GetPlatformFile().FileExists(*ProjectFile);
diff --git a/Engine/Source/Editor/GameProjectGeneration/Public/GameProjectUtils.h b/Engine/Source/Editor/GameProjectGeneration/Public/GameProjectUtils.h
--- a/Engine/Source/Editor/GameProjectGeneration/Public/GameProjectUtils.h
+++ b/Engine/Source/Editor/GameProjectGeneration/Public/GameProjectUtils.h
@@ -321,9 +321,6 @@ private:
 	/** Returns the include header path for a given fully specified, normalized file path */
 	static FString GetIncludePathForFile(const FString& InFullFilePath, const FString& ModuleRootPath);
 
-	/** Checks the name for an underscore and the existence of XB1 XDK */
-	static bool NameContainsUnderscoreAndXB1Installed(const FString& TestName);
-
 	/** Returns true if the project file exists on disk */
 	static bool ProjectFileExists(const FString& ProjectFile);
 
