Add the ability to enable re-instancing via project settings.  Old setting still supported.
Add the ability to enable auto compile when adding new C++ classes.
Fixed user feedback issues when adding new C++ classes.

#rnx

diff --git a/Engine/Config/BaseEditorPerProjectUserSettings.ini b/Engine/Config/BaseEditorPerProjectUserSettings.ini
--- a/Engine/Config/BaseEditorPerProjectUserSettings.ini
+++ b/Engine/Config/BaseEditorPerProjectUserSettings.ini
@@ -938,4 +938,7 @@ bShowTransactionDetails=false
 
 [/Script/WindowsMixedRealityRuntimeSettings.WindowsMixedRealityRuntimeSettings]
 RemoteHoloLensIP=
-MaxBitrate=4000
\ No newline at end of file
+MaxBitrate=4000
+
+[/Script/LiveCoding.LiveCodingSettings]
+bAutomaticallyCompileNewClasses=True
\ No newline at end of file
diff --git a/Engine/Source/Developer/Windows/LiveCoding/Private/LiveCodingModule.cpp b/Engine/Source/Developer/Windows/LiveCoding/Private/LiveCodingModule.cpp
--- a/Engine/Source/Developer/Windows/LiveCoding/Private/LiveCodingModule.cpp
+++ b/Engine/Source/Developer/Windows/LiveCoding/Private/LiveCodingModule.cpp
@@ -231,6 +231,7 @@ void FLiveCodingModule::StartupModule()
 	}
 
 	bEnabledLastTick = Settings->bEnabled;
+	bEnableReinstancingLastTick = IsReinstancingEnabled();
 }
 
 void FLiveCodingModule::ShutdownModule()
@@ -358,6 +359,11 @@ void FLiveCodingModule::Tick()
 			FMessageDialog::Open(EAppMsgType::Ok, LOCTEXT("NoEnableLiveCodingAfterHotReload", "Live Coding cannot be enabled while hot-reloaded modules are active. Please close the editor and build from your IDE before restarting."));
 		}
 	}
+	else if (IsEnabledForSession() && IsReinstancingEnabled() != bEnableReinstancingLastTick)
+	{
+		bEnableReinstancingLastTick = IsReinstancingEnabled();
+		LppSetReinstancingFlow(bEnableReinstancingLastTick);
+	}
 
 	if (bUpdateModulesInTick)
 	{
@@ -627,9 +633,7 @@ bool FLiveCodingModule::StartLiveCoding()
 		LppSetBuildArguments(*Arguments);
 
 #if WITH_EDITOR
-		bool bEnableReinstancing = true;
-		GConfig->GetBool(TEXT("LiveCoding"), TEXT("bEnableReinstancing"), bEnableReinstancing, GEngineIni);
-		if (bEnableReinstancing)
+		if (IsReinstancingEnabled())
 		{
 			LppSetReinstancingFlow(true);
 		}
@@ -813,10 +817,8 @@ void FLiveCodingModule::BeginReload()
 			GLiveCodingModule->bHasPatchBeenLoaded = false;
 			GPostCompileResult = commands::PostCompileResult::Success;
 #if WITH_EDITOR
-			bool bEnableReinstancing = true;
-			GConfig->GetBool(TEXT("LiveCoding"), TEXT("bEnableReinstancing"), bEnableReinstancing, GEngineIni);
 			GLiveCodingModule->Reload.Reset(new FReload(EActiveReloadType::LiveCoding, TEXT("LIVECODING"), *GLog));
-			GLiveCodingModule->Reload->SetEnableReinstancing(bEnableReinstancing);
+			GLiveCodingModule->Reload->SetEnableReinstancing(GLiveCodingModule->IsReinstancingEnabled());
 			GLiveCodingModule->Reload->SetSendReloadCompleteNotification(false);
 #else
 			GLiveCodingModule->Reload.Reset(new FNullReload(*GLiveCodingModule));
@@ -825,6 +827,26 @@ void FLiveCodingModule::BeginReload()
 	}
 }
 
+bool FLiveCodingModule::IsReinstancingEnabled() const
+{
+#if WITH_EDITOR
+	if (Settings->bEnableReinstancing)
+	{
+		return true;
+	}
+	bool bEnableReinstancing = true;
+	GConfig->GetBool(TEXT("LiveCoding"), TEXT("bEnableReinstancing"), bEnableReinstancing, GEngineIni);
+	return bEnableReinstancing;
+#else
+	return false;
+#endif
+}
+
+bool FLiveCodingModule::AutomaticallyCompileNewClasses() const
+{
+	return Settings->bAutomaticallyCompileNewClasses;
+}
+
 // Invoked from LC_ClientCommandActions
 void LiveCodingBeginPatch()
 {
diff --git a/Engine/Source/Developer/Windows/LiveCoding/Private/LiveCodingModule.h b/Engine/Source/Developer/Windows/LiveCoding/Private/LiveCodingModule.h
--- a/Engine/Source/Developer/Windows/LiveCoding/Private/LiveCodingModule.h
+++ b/Engine/Source/Developer/Windows/LiveCoding/Private/LiveCodingModule.h
@@ -33,6 +33,7 @@ public:
 	virtual bool IsEnabledByDefault() const override;
 	virtual void EnableForSession(bool bInEnabled) override;
 	virtual bool IsEnabledForSession() const override;
+	virtual bool AutomaticallyCompileNewClasses() const override;
 	virtual bool CanEnableForSession() const override;
 	virtual bool HasStarted() const override;
 	virtual void ShowConsole() override;
@@ -50,6 +51,7 @@ private:
 	ULiveCodingSettings* Settings;
 	TSharedPtr<ISettingsSection> SettingsSection;
 	bool bEnabledLastTick = false;
+	bool bEnableReinstancingLastTick = false;
 	bool bEnabledForSession = false;
 	bool bStarted = false;
 	bool bUpdateModulesInTick = false;
@@ -84,6 +86,8 @@ private:
 
 	bool ShouldPreloadModule(const FName& Name, const FString& FullFilePath) const;
 
+	bool IsReinstancingEnabled() const;
+
 #if WITH_EDITOR
 	void ShowNotification(bool Success, const FText& Title, const FText* SubText);
 #endif
diff --git a/Engine/Source/Developer/Windows/LiveCoding/Private/LiveCodingSettings.h b/Engine/Source/Developer/Windows/LiveCoding/Private/LiveCodingSettings.h
--- a/Engine/Source/Developer/Windows/LiveCoding/Private/LiveCodingSettings.h
+++ b/Engine/Source/Developer/Windows/LiveCoding/Private/LiveCodingSettings.h
@@ -26,6 +26,12 @@ public:
     UPROPERTY(config, EditAnywhere, Category=General, Meta=(ConfigRestartRequired=true, EditCondition="bEnabled"))
     ELiveCodingStartupMode Startup;
 
+    UPROPERTY(config, EditAnywhere, Category=General, Meta=(EditCondition="bEnabled"))
+	bool bEnableReinstancing;
+
+	UPROPERTY(config, EditAnywhere, Category=General, Meta=(EditCondition="bEnabled", DisplayName="Automatically Compile Newly Added C++ Classes"))
+	bool bAutomaticallyCompileNewClasses;
+
 	UPROPERTY(config, EditAnywhere, Category=Modules, Meta=(ConfigRestartRequired=true, EditCondition="bEnabled"))
     bool bPreloadEngineModules;
 
diff --git a/Engine/Source/Developer/Windows/LiveCoding/Public/ILiveCodingModule.h b/Engine/Source/Developer/Windows/LiveCoding/Public/ILiveCodingModule.h
--- a/Engine/Source/Developer/Windows/LiveCoding/Public/ILiveCodingModule.h
+++ b/Engine/Source/Developer/Windows/LiveCoding/Public/ILiveCodingModule.h
@@ -17,6 +17,7 @@ public:
 	virtual void EnableForSession(bool bEnabled) = 0;
 	virtual bool IsEnabledForSession() const = 0;
 	virtual bool CanEnableForSession() const = 0;
+	virtual bool AutomaticallyCompileNewClasses() const = 0;
 
 	virtual bool HasStarted() const = 0;
 
diff --git a/Engine/Source/Editor/GameProjectGeneration/Private/GameProjectUtils.cpp b/Engine/Source/Editor/GameProjectGeneration/Private/GameProjectUtils.cpp
--- a/Engine/Source/Editor/GameProjectGeneration/Private/GameProjectUtils.cpp
+++ b/Engine/Source/Editor/GameProjectGeneration/Private/GameProjectUtils.cpp
@@ -1327,7 +1327,13 @@ bool GameProjectUtils::IsValidBaseClassForCreation_Internal(const UClass* InClas
 
 GameProjectUtils::EAddCodeToProjectResult GameProjectUtils::AddCodeToProject(const FString& NewClassName, const FString& NewClassPath, const FModuleContextInfo& ModuleInfo, const FNewClassInfo ParentClassInfo, const TSet<FString>& DisallowedHeaderNames, FString& OutHeaderFilePath, FString& OutCppFilePath, FText& OutFailReason)
 {
-	const EAddCodeToProjectResult Result = AddCodeToProject_Internal(NewClassName, NewClassPath, ModuleInfo, ParentClassInfo, DisallowedHeaderNames, OutHeaderFilePath, OutCppFilePath, OutFailReason);
+	EReloadStatus OutReloadStatus;
+	return AddCodeToProject(NewClassName, NewClassPath, ModuleInfo, ParentClassInfo, DisallowedHeaderNames, OutHeaderFilePath, OutCppFilePath, OutFailReason, OutReloadStatus);
+}
+
+GameProjectUtils::EAddCodeToProjectResult GameProjectUtils::AddCodeToProject(const FString& NewClassName, const FString& NewClassPath, const FModuleContextInfo& ModuleInfo, const FNewClassInfo ParentClassInfo, const TSet<FString>& DisallowedHeaderNames, FString& OutHeaderFilePath, FString& OutCppFilePath, FText& OutFailReason, EReloadStatus& OutReloadStatus)
+{
+	const EAddCodeToProjectResult Result = AddCodeToProject_Internal(NewClassName, NewClassPath, ModuleInfo, ParentClassInfo, DisallowedHeaderNames, OutHeaderFilePath, OutCppFilePath, OutFailReason, OutReloadStatus);
 
 	if( FEngineAnalytics::IsAvailable() )
 	{
@@ -3968,7 +3974,7 @@ TArray<FString> GameProjectUtils::GetRequiredAdditionalDependencies(const FNewCl
 	return Out;
 }
 
-GameProjectUtils::EAddCodeToProjectResult GameProjectUtils::AddCodeToProject_Internal(const FString& NewClassName, const FString& NewClassPath, const FModuleContextInfo& ModuleInfo, const FNewClassInfo ParentClassInfo, const TSet<FString>& DisallowedHeaderNames, FString& OutHeaderFilePath, FString& OutCppFilePath, FText& OutFailReason)
+GameProjectUtils::EAddCodeToProjectResult GameProjectUtils::AddCodeToProject_Internal(const FString& NewClassName, const FString& NewClassPath, const FModuleContextInfo& ModuleInfo, const FNewClassInfo ParentClassInfo, const TSet<FString>& DisallowedHeaderNames, FString& OutHeaderFilePath, FString& OutCppFilePath, FText& OutFailReason, EReloadStatus& OutReloadStatus)
 {
 	if ( !ParentClassInfo.IsSet() )
 	{
@@ -4133,13 +4139,18 @@ GameProjectUtils::EAddCodeToProjectResult GameProjectUtils::AddCodeToProject_Int
 
 	OutHeaderFilePath = NewHeaderFilename;
 	OutCppFilePath = NewCppFilename;
+	OutReloadStatus = EReloadStatus::NotReloaded;
 
 #if WITH_LIVE_CODING
 	ILiveCodingModule* LiveCoding = FModuleManager::GetModulePtr<ILiveCodingModule>(LIVE_CODING_MODULE_NAME);
 	if (LiveCoding != nullptr && LiveCoding->IsEnabledForSession())
 	{
-		OutFailReason = LOCTEXT("FailedToCompileLiveCodingEnabled", "Adding classes dynamically is not allowed with Live Coding enabled.");
-		return EAddCodeToProjectResult::FailedToHotReload;
+		if (bProjectHadCodeFiles && LiveCoding->AutomaticallyCompileNewClasses())
+		{
+			LiveCoding->Compile();
+			OutReloadStatus = EReloadStatus::Reloaded;
+		}
+		return EAddCodeToProjectResult::Succeeded;
 	}
 #endif
 
@@ -4163,6 +4174,7 @@ GameProjectUtils::EAddCodeToProjectResult GameProjectUtils::AddCodeToProject_Int
 
 		// Notify that we've created a brand new module
 		FSourceCodeNavigation::AccessOnNewModuleAdded().Broadcast(*GameModuleName);
+		OutReloadStatus = EReloadStatus::Reloaded;
 	}
 	else if (GetDefault<UEditorPerProjectUserSettings>()->bAutomaticallyHotReloadNewClasses)
 	{
@@ -4208,6 +4220,7 @@ GameProjectUtils::EAddCodeToProjectResult GameProjectUtils::AddCodeToProject_Int
 				}
 			}
 		}
+		OutReloadStatus = EReloadStatus::Reloaded;
 	}
 
 	return EAddCodeToProjectResult::Succeeded;
diff --git a/Engine/Source/Editor/GameProjectGeneration/Private/SNewClassDialog.cpp b/Engine/Source/Editor/GameProjectGeneration/Private/SNewClassDialog.cpp
--- a/Engine/Source/Editor/GameProjectGeneration/Private/SNewClassDialog.cpp
+++ b/Engine/Source/Editor/GameProjectGeneration/Private/SNewClassDialog.cpp
@@ -1163,9 +1163,10 @@ void SNewClassDialog::FinishClicked()
 		// Track the selected module name so we can default to this next time
 		LastSelectedModuleName = SelectedModuleInfo->ModuleName;
 
+		GameProjectUtils::EReloadStatus ReloadStatus;
 		FText FailReason;
 		const TSet<FString>& DisallowedHeaderNames = FSourceCodeNavigation::GetSourceFileDatabase().GetDisallowedHeaderNames();
-		const GameProjectUtils::EAddCodeToProjectResult AddCodeResult = GameProjectUtils::AddCodeToProject(NewClassName, NewClassPath, *SelectedModuleInfo, ParentClassInfo, DisallowedHeaderNames, HeaderFilePath, CppFilePath, FailReason);
+		const GameProjectUtils::EAddCodeToProjectResult AddCodeResult = GameProjectUtils::AddCodeToProject(NewClassName, NewClassPath, *SelectedModuleInfo, ParentClassInfo, DisallowedHeaderNames, HeaderFilePath, CppFilePath, FailReason, ReloadStatus);
 		if (AddCodeResult == GameProjectUtils::EAddCodeToProjectResult::Succeeded)
 		{
 			OnAddedToProject.ExecuteIfBound( NewClassName, NewClassPath, SelectedModuleInfo->ModuleName );
@@ -1177,8 +1178,9 @@ void SNewClassDialog::FinishClicked()
 			bPreventPeriodicValidityChecksUntilNextChange = true;
 
 			// Display a nag if we didn't automatically hot-reload for the newly added class
-			const bool bWasHotReloaded = GetDefault<UEditorPerProjectUserSettings>()->bAutomaticallyHotReloadNewClasses;
-			if( bWasHotReloaded )
+			bool bWasReloaded = ReloadStatus == GameProjectUtils::EReloadStatus::Reloaded;
+
+			if( bWasReloaded )
 			{
 				FNotificationInfo Notification( FText::Format( LOCTEXT("AddedClassSuccessNotification", "Added new class {0}"), FText::FromString(NewClassName) ) );
 				FSlateNotificationManager::Get().AddNotification( Notification );
@@ -1186,7 +1188,7 @@ void SNewClassDialog::FinishClicked()
 
 			if ( HeaderFilePath.IsEmpty() || CppFilePath.IsEmpty() || !FSlateApplication::Get().SupportsSourceAccess() )
 			{
-				if( !bWasHotReloaded )
+				if( !bWasReloaded )
 				{
 					// Code successfully added, notify the user. We are either running on a platform that does not support source access or a file was not given so don't ask about editing the file
 					const FText Message = FText::Format( 
@@ -1202,7 +1204,7 @@ void SNewClassDialog::FinishClicked()
 			else
 			{
 				bool bEditSourceFilesNow = false;
-				if( bWasHotReloaded )
+				if( bWasReloaded )
 				{
 					// Code was hot reloaded, so always edit the new classes now
 					bEditSourceFilesNow = true;
diff --git a/Engine/Source/Editor/GameProjectGeneration/Public/GameProjectUtils.h b/Engine/Source/Editor/GameProjectGeneration/Public/GameProjectUtils.h
--- a/Engine/Source/Editor/GameProjectGeneration/Public/GameProjectUtils.h
+++ b/Engine/Source/Editor/GameProjectGeneration/Public/GameProjectUtils.h
@@ -91,6 +91,16 @@ public:
 		UserCanceled
 	};
 
+	/** Reload status when adding new code to the project */
+	enum class EReloadStatus : uint8
+	{
+		/** Code was built and reloaded */
+		Reloaded,
+
+		/** Code was not reloaded */
+		NotReloaded,
+	};
+
 	/** Returns true if the project filename is properly formed and does not conflict with another project */
 	static bool IsValidProjectFileForCreation(const FString& ProjectFile, FText& OutFailReason);
 
@@ -138,6 +148,9 @@ public:
 	/** Adds new source code to the project. When returning Succeeded or FailedToHotReload, OutSyncFileAndLineNumber will be the the preferred target file to sync in the users code editing IDE, formatted for use with GenericApplication::GotoLineInSource */
 	static EAddCodeToProjectResult AddCodeToProject(const FString& NewClassName, const FString& NewClassPath, const FModuleContextInfo& ModuleInfo, const FNewClassInfo ParentClassInfo, const TSet<FString>& DisallowedHeaderNames, FString& OutHeaderFilePath, FString& OutCppFilePath, FText& OutFailReason);
 
+	/** Adds new source code to the project. When returning Succeeded or FailedToHotReload, OutSyncFileAndLineNumber will be the the preferred target file to sync in the users code editing IDE, formatted for use with GenericApplication::GotoLineInSource */
+	static EAddCodeToProjectResult AddCodeToProject(const FString& NewClassName, const FString& NewClassPath, const FModuleContextInfo& ModuleInfo, const FNewClassInfo ParentClassInfo, const TSet<FString>& DisallowedHeaderNames, FString& OutHeaderFilePath, FString& OutCppFilePath, FText& OutFailReason, EReloadStatus& OutReloadStatus);
+
 	/** Loads a list of template categories defined in the TemplateCategories.ini file in the specified folder */
 	static UTemplateCategories* LoadTemplateCategories(const FString& RootDir);
 
@@ -455,7 +468,7 @@ private:
 	static bool CheckoutGameProjectFile(const FString& ProjectFilename, FText& OutFailReason);
 
 	/** Internal handler for AddCodeToProject*/
-	static EAddCodeToProjectResult AddCodeToProject_Internal(const FString& NewClassName, const FString& NewClassPath, const FModuleContextInfo& ModuleInfo, const FNewClassInfo ParentClassInfo, const TSet<FString>& DisallowedHeaderNames, FString& OutHeaderFilePath, FString& OutCppFilePath, FText& OutFailReason);
+	static EAddCodeToProjectResult AddCodeToProject_Internal(const FString& NewClassName, const FString& NewClassPath, const FModuleContextInfo& ModuleInfo, const FNewClassInfo ParentClassInfo, const TSet<FString>& DisallowedHeaderNames, FString& OutHeaderFilePath, FString& OutCppFilePath, FText& OutFailReason, EReloadStatus& OutReloadStatus);
 
 	/** Internal handler for IsValidBaseClassForCreation */
 	DECLARE_DELEGATE_RetVal_OneParam(bool, FDoesClassNeedAPIExportCallback, const FString& /*ClassModuleName*/);
