Don't copy over saved config folder when copying projects.

diff --git a/Engine/Source/Editor/GameProjectGeneration/Private/GameProjectUtils.cpp b/Engine/Source/Editor/GameProjectGeneration/Private/GameProjectUtils.cpp
--- a/Engine/Source/Editor/GameProjectGeneration/Private/GameProjectUtils.cpp
+++ b/Engine/Source/Editor/GameProjectGeneration/Private/GameProjectUtils.cpp
@@ -2840,6 +2840,7 @@ GameProjectUtils::EProjectDuplicateResult GameProjectUtils::DuplicateProjectForU
 				TEXT("Saved/Autosaves"),
 				TEXT("Saved/Backup"),
 				TEXT("Saved/Cooked"),
+				TEXT("Saved/Config"),
 				TEXT("Saved/HardwareSurvey"),
 				TEXT("Saved/Logs"),
 				TEXT("Saved/StagedBuilds"),
