Fix missing ' for WorldPartitionConvertCommandlet isn't set as the default Commandlet Class in newly-created projects

diff --git a/Engine/Source/Editor/GameProjectGeneration/Private/GameProjectUtils.cpp b/Engine/Source/Editor/GameProjectGeneration/Private/GameProjectUtils.cpp
--- a/Engine/Source/Editor/GameProjectGeneration/Private/GameProjectUtils.cpp
+++ b/Engine/Source/Editor/GameProjectGeneration/Private/GameProjectUtils.cpp
@@ -288,7 +288,7 @@ namespace
 			ConfigValues.Emplace(TEXT("DefaultEngine.ini"),
 				TEXT("/Script/WorldPartitionEditor.WorldPartitionEditorSettings"),
 				TEXT("CommandletClass"),
-				TEXT("Class'/Script/UnrealEd.WorldPartitionConvertCommandlet"),
+				TEXT("Class'/Script/UnrealEd.WorldPartitionConvertCommandlet'"),
 				true /* ShouldReplaceExistingValue */);
 		}
 	}
