Removed "Fail If Generated Code Changes" requirement for this case.  Adding new types will cause the project generated file to change thus triggering the failure.

#rnx

diff --git a/Engine/Source/Editor/GameProjectGeneration/Private/GameProjectUtils.cpp b/Engine/Source/Editor/GameProjectGeneration/Private/GameProjectUtils.cpp
--- a/Engine/Source/Editor/GameProjectGeneration/Private/GameProjectUtils.cpp
+++ b/Engine/Source/Editor/GameProjectGeneration/Private/GameProjectUtils.cpp
@@ -4214,7 +4214,7 @@ GameProjectUtils::EAddCodeToProjectResult GameProjectUtils::AddCodeToProject_Int
 			else
 			{
 				// Perform a regular unload, then reload
-				if (!HotReloadSupport.RecompileModule(ModuleFName, *GWarn, ERecompileModuleFlags::ReloadAfterRecompile | ERecompileModuleFlags::FailIfGeneratedCodeChanges))
+				if (!HotReloadSupport.RecompileModule(ModuleFName, *GWarn, ERecompileModuleFlags::ReloadAfterRecompile))
 				{
 					OutFailReason = FText::Format(LOCTEXT("FailedToCompileModuleFmt", "Failed to automatically compile the '{0}' module."), FText::FromString(ModuleInfo.ModuleName));
 					return EAddCodeToProjectResult::FailedToHotReload;
