Add a button to disable the current IDE if an IDE is not required to compile for the Engine.

This is mainly for Linux only, as the toolchain required to setup Linux comes with the compiler needed. So an IDE is not required

diff --git a/Engine/Source/Editor/GameProjectGeneration/Private/SGetSuggestedIDEWidget.cpp b/Engine/Source/Editor/GameProjectGeneration/Private/SGetSuggestedIDEWidget.cpp
--- a/Engine/Source/Editor/GameProjectGeneration/Private/SGetSuggestedIDEWidget.cpp
+++ b/Engine/Source/Editor/GameProjectGeneration/Private/SGetSuggestedIDEWidget.cpp
@@ -14,6 +14,33 @@
 
 #define LOCTEXT_NAMESPACE "GameProjectGeneration"
 
+void SGetDisableIDEWidget::Construct(const FArguments& InArgs)
+{
+	ChildSlot
+	[
+		CreateGetDisableIDEWidget()
+	];
+}
+
+TSharedRef<SWidget> SGetDisableIDEWidget::CreateGetDisableIDEWidget()
+{
+	return
+		SNew(SButton)
+		.HAlign(HAlign_Center)
+		.VAlign(VAlign_Center)
+		.ButtonStyle(FAppStyle::Get(), "PrimaryButton")
+		.TextStyle(FAppStyle::Get(), "DialogButtonText")
+		.Text(FText::Format(LOCTEXT("IDEDisableButton", "Disable {0}"), FSourceCodeNavigation::GetSuggestedSourceCodeIDE()))
+		.OnClicked(this, &SGetDisableIDEWidget::OnDisableIDEClicked);
+}
+
+FReply SGetDisableIDEWidget::OnDisableIDEClicked()
+{
+	FSourceCodeNavigation::SetPreferredAccessor(TEXT("NullSourceCodeAccessor"));
+
+	return FReply::Handled();
+}
+
 TSharedPtr<SNotificationItem> SGetSuggestedIDEWidget::IDEDownloadNotification;
 
 void SGetSuggestedIDEWidget::Construct(const FArguments& InArgs)
diff --git a/Engine/Source/Editor/GameProjectGeneration/Private/SGetSuggestedIDEWidget.h b/Engine/Source/Editor/GameProjectGeneration/Private/SGetSuggestedIDEWidget.h
--- a/Engine/Source/Editor/GameProjectGeneration/Private/SGetSuggestedIDEWidget.h
+++ b/Engine/Source/Editor/GameProjectGeneration/Private/SGetSuggestedIDEWidget.h
@@ -47,3 +47,24 @@ private:
 	static TSharedPtr<class SNotificationItem> IDEDownloadNotification;
 
 };
+
+/**
+ * A button to disable requiring an IDE to be install 
+ * Only visible on platforms that dont require an IDE to compile code and when no IDE is available
+ */
+class SGetDisableIDEWidget : public SCompoundWidget
+{
+public:
+	SLATE_BEGIN_ARGS(SGetDisableIDEWidget) {}
+	SLATE_END_ARGS()
+
+	/** Constructs this widget with InArgs */
+	void Construct(const FArguments& InArgs);
+	
+private:
+	/** Creates the appropriate widget to display for the platform */
+	TSharedRef<SWidget> CreateGetDisableIDEWidget();
+
+	/** Handler for when the disable button is clicked */
+	FReply OnDisableIDEClicked();
+};
diff --git a/Engine/Source/Editor/GameProjectGeneration/Private/SNewClassDialog.cpp b/Engine/Source/Editor/GameProjectGeneration/Private/SNewClassDialog.cpp
--- a/Engine/Source/Editor/GameProjectGeneration/Private/SNewClassDialog.cpp
+++ b/Engine/Source/Editor/GameProjectGeneration/Private/SNewClassDialog.cpp
@@ -988,7 +988,11 @@ FText SNewClassDialog::GetGlobalErrorLabelText() const
 {
 	if ( ClassDomain == EClassDomain::Native && !FSourceCodeNavigation::IsCompilerAvailable() )
 	{
-		return FText::Format( LOCTEXT("NoCompilerFound", "No compiler was found. In order to use C++ code, you must first install {0}."), FSourceCodeNavigation::GetSuggestedSourceCodeIDE() );
+#if PLATFORM_LINUX
+		return FText::Format(LOCTEXT("NoCompilerFound", "In order to use a C++ template, you must first install or disable {0}."), FSourceCodeNavigation::GetSuggestedSourceCodeIDE());
+#elif
+		return FText::Format(LOCTEXT("NoCompilerFound", "No compiler was found. In order to use C++ code, you must first install {0}."), FSourceCodeNavigation::GetSuggestedSourceCodeIDE());
+#endif
 	}
 
 	return FText::GetEmpty();
diff --git a/Engine/Source/Editor/GameProjectGeneration/Private/SProjectDialog.cpp b/Engine/Source/Editor/GameProjectGeneration/Private/SProjectDialog.cpp
--- a/Engine/Source/Editor/GameProjectGeneration/Private/SProjectDialog.cpp
+++ b/Engine/Source/Editor/GameProjectGeneration/Private/SProjectDialog.cpp
@@ -510,6 +510,14 @@ TSharedRef<SWidget> SProjectDialog::MakeNewProjectDialogButtons()
 			SNew(SGetSuggestedIDEWidget)
 			.VisibilityOverride(this, &SProjectDialog::GetSuggestedIDEButtonVisibility)
 		]
+		+ SHorizontalBox::Slot()
+		.Padding(8.0f, 0.0f, 8.0f, 0.0f)
+		.VAlign(VAlign_Center)
+		.AutoWidth()
+		[
+			SNew(SGetDisableIDEWidget)
+			.Visibility(this, &SProjectDialog::GetDisableIDEButtonVisibility)
+		]
 		+SHorizontalBox::Slot()
 		.Padding(8.0f, 0.0f, 8.0f, 0.0f)
 		.VAlign(VAlign_Center)
@@ -1850,6 +1858,17 @@ EVisibility SProjectDialog::GetSuggestedIDEButtonVisibility() const
 	return IsCompilerRequired() && !FSourceCodeNavigation::IsCompilerAvailable() ? EVisibility::Visible : EVisibility::Collapsed;
 }
 
+// Allow disabling of the current IDE for platforms that dont require an IDE to run the Editor/Engine
+EVisibility SProjectDialog::GetDisableIDEButtonVisibility() const
+{
+	if (GetSuggestedIDEButtonVisibility() == EVisibility::Visible && !IsIDERequired())
+	{
+		return EVisibility::Visible;
+	}
+
+	return EVisibility::Collapsed;
+}
+
 const FSlateBrush* SProjectDialog::GetSelectedTemplatePreviewImage() const
 {
 	TSharedPtr<FSlateBrush> PreviewImage = GetSelectedTemplateProperty(&FTemplateItem::PreviewImage);
@@ -1922,6 +1941,15 @@ bool SProjectDialog::IsCompilerRequired() const
 	return false;
 }
 
+// Linux does not require an IDE to be setup to compile things
+bool SProjectDialog::IsIDERequired() const
+{
+#if PLATFORM_LINUX
+	return false;
+#endif
+	return true;
+}
+
 EVisibility SProjectDialog::GetProjectSettingsVisibility() const
 {
 	const TArray<ETemplateSetting>& HiddenSettings = GetSelectedTemplateProperty(&FTemplateItem::HiddenSettings);
@@ -1985,7 +2013,15 @@ void SProjectDialog::UpdateProjectFileValidity()
 				if (!FSourceCodeNavigation::IsCompilerAvailable())
 				{
 					bLastGlobalValidityCheckSuccessful = false;
-					LastGlobalValidityErrorText = FText::Format(LOCTEXT("NoCompilerFound", "No compiler was found. In order to use a C++ template, you must first install {0}."), FSourceCodeNavigation::GetSuggestedSourceCodeIDE());
+
+					if (IsIDERequired())
+					{
+						LastGlobalValidityErrorText = FText::Format(LOCTEXT("NoCompilerFound", "No compiler was found. In order to use a C++ template, you must first install {0}."), FSourceCodeNavigation::GetSuggestedSourceCodeIDE());
+					}
+					else
+					{
+						LastGlobalValidityErrorText = FText::Format(LOCTEXT("NoCompilerFound", "In order to use a C++ template, you must first install or disable {0}."), FSourceCodeNavigation::GetSuggestedSourceCodeIDE());
+					}
 				}
 				else if (!FDesktopPlatformModule::Get()->IsUnrealBuildToolAvailable())
 				{
diff --git a/Engine/Source/Editor/GameProjectGeneration/Private/SProjectDialog.h b/Engine/Source/Editor/GameProjectGeneration/Private/SProjectDialog.h
--- a/Engine/Source/Editor/GameProjectGeneration/Private/SProjectDialog.h
+++ b/Engine/Source/Editor/GameProjectGeneration/Private/SProjectDialog.h
@@ -77,6 +77,8 @@ private:
 	/** Returns true if we have a code template selected */
 	bool IsCompilerRequired() const;
 
+	bool IsIDERequired() const;
+
 	EVisibility GetProjectSettingsVisibility() const;
 
 	/** Get a visibility of the class types display. If the string is empty this return Collapsed otherwise it will return Visible */
@@ -91,6 +93,8 @@ private:
 
 	EVisibility GetSuggestedIDEButtonVisibility() const;
 
+	EVisibility GetDisableIDEButtonVisibility() const;
+
 	/** Accessor for the project name text */
 	FText GetCurrentProjectFileName() const;
 
@@ -235,4 +239,4 @@ private:
 	/** True if the last NameAndLocation validity check returned that the project path is valid for creation */
 	bool bLastNameAndLocationValidityCheckSuccessful;
 
-};
\ No newline at end of file
+};
diff --git a/Engine/Source/Editor/UnrealEd/Private/SourceCodeNavigation.cpp b/Engine/Source/Editor/UnrealEd/Private/SourceCodeNavigation.cpp
--- a/Engine/Source/Editor/UnrealEd/Private/SourceCodeNavigation.cpp
+++ b/Engine/Source/Editor/UnrealEd/Private/SourceCodeNavigation.cpp
@@ -1506,6 +1506,16 @@ void FSourceCodeNavigation::RemoveNavigationHandler(ISourceCodeNavigationHandler
 	SourceCodeNavigationHandlers.Remove(Handler);
 }
 
+void FSourceCodeNavigation::SetPreferredAccessor(const TCHAR* Name)
+{
+	GConfig->SetString(TEXT("/Script/SourceCodeAccess.SourceCodeAccessSettings"), TEXT("PreferredAccessor"), Name, GEngineIni);
+
+	ISourceCodeAccessModule& SourceCodeAccessModule = FModuleManager::LoadModuleChecked<ISourceCodeAccessModule>(TEXT("SourceCodeAccess"));
+	SourceCodeAccessModule.SetAccessor(Name);
+
+	RefreshCompilerAvailability();
+}
+
 bool FSourceCodeNavigation::CanNavigateToClass(const UClass* InClass)
 {
 	if (!InClass)
@@ -1810,7 +1820,7 @@ FText FSourceCodeNavigation::GetSuggestedSourceCodeIDE(bool bShortIDEName)
 #elif PLATFORM_MAC
 	return LOCTEXT("SuggestedCodeIDE_Mac", "Xcode");
 #elif PLATFORM_LINUX
-	return LOCTEXT("SuggestedCodeIDE_Linux", "Visual Studio Code");
+	return LOCTEXT("SuggestedCodeIDE_Linux", "VS Code");
 #else
 	return LOCTEXT("SuggestedCodeIDE_Generic", "an IDE to edit source code");
 #endif
diff --git a/Engine/Source/Editor/UnrealEd/Public/SourceCodeNavigation.h b/Engine/Source/Editor/UnrealEd/Public/SourceCodeNavigation.h
--- a/Engine/Source/Editor/UnrealEd/Public/SourceCodeNavigation.h
+++ b/Engine/Source/Editor/UnrealEd/Public/SourceCodeNavigation.h
@@ -390,6 +390,9 @@ public:
 	/** Remove a navigation handler */
 	UNREALED_API static void RemoveNavigationHandler(ISourceCodeNavigationHandler* Handler);
 
+	/** Set a new preferred accessor and refresh compiler status */
+	UNREALED_API static void SetPreferredAccessor(const TCHAR* Name);
+
 private:
 
 	/** Critical section for locking access to the source file database. */
