Fix compiler error

diff --git a/Engine/Source/Editor/GameProjectGeneration/Private/SNewClassDialog.cpp b/Engine/Source/Editor/GameProjectGeneration/Private/SNewClassDialog.cpp
--- a/Engine/Source/Editor/GameProjectGeneration/Private/SNewClassDialog.cpp
+++ b/Engine/Source/Editor/GameProjectGeneration/Private/SNewClassDialog.cpp
@@ -990,7 +990,7 @@ FText SNewClassDialog::GetGlobalErrorLabelText() const
 	{
 #if PLATFORM_LINUX
 		return FText::Format(LOCTEXT("NoCompilerFound", "In order to use a C++ template, you must first install or disable {0}."), FSourceCodeNavigation::GetSuggestedSourceCodeIDE());
-#elif
+#else
 		return FText::Format(LOCTEXT("NoCompilerFound", "No compiler was found. In order to use C++ code, you must first install {0}."), FSourceCodeNavigation::GetSuggestedSourceCodeIDE());
 #endif
 	}
