Fix project dialog marketplace button alignment

diff --git a/Engine/Source/Editor/GameProjectGeneration/Private/SProjectDialog.cpp b/Engine/Source/Editor/GameProjectGeneration/Private/SProjectDialog.cpp
--- a/Engine/Source/Editor/GameProjectGeneration/Private/SProjectDialog.cpp
+++ b/Engine/Source/Editor/GameProjectGeneration/Private/SProjectDialog.cpp
@@ -555,12 +555,9 @@ TSharedRef<SWidget> SProjectDialog::MakeOpenProjectDialogButtons()
 		.AutoWidth()
 		[
 			SNew(SButton)
-			.HAlign(HAlign_Center)
-			.VAlign(VAlign_Center)
 			.ToolTipText(LOCTEXT("MarketplaceToolTip", "Check out the Marketplace to find new projects!"))
 			.Visibility(FLauncherPlatformModule::Get()->CanOpenLauncher(true) ? EVisibility::Visible : EVisibility::Collapsed)
 			.OnClicked(ProjectBrowserRef, &SProjectBrowser::OnOpenMarketplace)
-			.ContentPadding(FMargin(4.0f, 2.0f,4.0f, 1.0f))
 			[	
 				SNew(SHorizontalBox)
 				+SHorizontalBox::Slot()
