Project browser now shows engine version for tiles at the current version to avoid confusion about what projects are at the latest. There are now sorting options for access time, alphabetical, and engine version

diff --git a/Engine/Source/Editor/GameProjectGeneration/Private/SProjectBrowser.cpp b/Engine/Source/Editor/GameProjectGeneration/Private/SProjectBrowser.cpp
--- a/Engine/Source/Editor/GameProjectGeneration/Private/SProjectBrowser.cpp
+++ b/Engine/Source/Editor/GameProjectGeneration/Private/SProjectBrowser.cpp
@@ -51,6 +51,8 @@
 #include "LauncherPlatformModule.h"
 #include "Interfaces/IMainFrameModule.h"
 #include "Internationalization/BreakIterator.h"
+#include "SSimpleButton.h"
+#include "Widgets/Input/SComboButton.h"
 
 #define LOCTEXT_NAMESPACE "ProjectBrowser"
 
@@ -63,6 +65,7 @@ namespace ProjectBrowserDefs
 	constexpr float ThumbnailSize = 64.0f, ThumbnailPadding = 5.f;
 }
 
+
 /**
  * Structure for project items.
  */
@@ -72,6 +75,7 @@ struct FProjectItem
 	FText Description;
 	FText Category;
 	FString EngineIdentifier;
+	FEngineVersion ProjectVersion;
 	FString ProjectFile;
 	TArray<FName> TargetPlatforms;
 	TSharedPtr<FSlateBrush> ProjectThumbnail;
@@ -88,9 +92,25 @@ struct FProjectItem
 		, TargetPlatforms(InTargetPlatforms)
 		, ProjectThumbnail(InProjectThumbnail)
 		, bUpToDate(InUpToDate)
-		//, bIsNewProjectItem(InIsNewProjectItem)
 		, bSupportsAllPlatforms(InSupportsAllPlatforms)
-	{}
+	{
+		if(bUpToDate)
+		{
+			ProjectVersion = FEngineVersion::Current();
+		}
+		else if (FDesktopPlatformModule::Get()->IsStockEngineRelease(EngineIdentifier))
+		{
+			FDesktopPlatformModule::Get()->TryParseStockEngineVersion(EngineIdentifier, ProjectVersion);
+		}
+
+
+		if (ProjectVersion.IsEmpty())
+		{
+			FString RootDir;
+			FDesktopPlatformModule::Get()->GetEngineRootDirFromIdentifier(EngineIdentifier, RootDir);
+			FDesktopPlatformModule::Get()->TryGetEngineVersion(RootDir, ProjectVersion);
+		}
+	}
 
 	/** Check if this project is up to date */
 	bool IsUpToDate() const
@@ -99,21 +119,31 @@ struct FProjectItem
 	}
 
 	/** Gets the engine label for this project */
-	FString GetEngineLabel() const
+	FText GetEngineLabel() const
 	{
 		if(bUpToDate)
 		{
-			return FString();
+			return LOCTEXT("CurrentEngineVersion", "Current");
 		}
 		else if(FDesktopPlatformModule::Get()->IsStockEngineRelease(EngineIdentifier))
 		{
-			return EngineIdentifier;
+			return FText::FromString(EngineIdentifier);
+		}
+		else if(!ProjectVersion.IsEmpty())
+		{
+			return FText::FromString(ProjectVersion.ToString(EVersionComponent::Patch));
 		}
 		else
 		{
-			return FString(TEXT("?"));
+			return FText::FromString(TEXT("?"));
 		}
 	}
+
+	bool CompareEngineVersion(const FProjectItem& Other) const
+	{
+		EVersionComponent Component;
+		return FEngineVersion::GetNewest(ProjectVersion, Other.ProjectVersion, &Component) == EVersionComparison::First ? true : false;
+	}
 };
 
 
@@ -186,10 +216,9 @@ public:
 							.VAlign(VAlign_Bottom)
 							[
 								SNew(STextBlock)
-								.Text(FText::FromString(ProjectItem->GetEngineLabel()))
+								.Text(ProjectItem->GetEngineLabel())
 								.Font(FAppStyle::Get().GetFontStyle("ProjectBrowser.ProjectTile.Font"))
 								.ColorAndOpacity(FAppStyle::Get().GetSlateColor("Colors.White25"))
-								.Visibility(ProjectItem->IsUpToDate() ? EVisibility::Collapsed : EVisibility::Visible)
 							]
 						]
 					]
@@ -290,6 +319,23 @@ void SProjectBrowser::Construct( const FArguments& InArgs )
 		ProjectTileView.ToSharedRef()
 	];
 
+	FMenuBuilder SortOptionsBuilder(true, nullptr);
+
+	auto MakeSortMenuEntry = [this, &SortOptionsBuilder](EProjectSortOption Option, FText Label, FText Tooltip)
+	{
+		FUIAction Action;
+	
+		Action.ExecuteAction = FExecuteAction::CreateSP(this, &SProjectBrowser::SortProjectTiles, Option);
+		Action.GetActionCheckState = FGetActionCheckState::CreateSP(this, &SProjectBrowser::GetSortOptionCheckState, Option);
+		SortOptionsBuilder.AddMenuEntry(Label, Tooltip, FSlateIcon(), Action, NAME_None, EUserInterfaceActionType::RadioButton);
+		
+	};
+
+	MakeSortMenuEntry(EProjectSortOption::LastAccessTime, LOCTEXT("ProjectSortMode_LastAcessTime", "Most Recently Used"), FText::GetEmpty());
+	MakeSortMenuEntry(EProjectSortOption::Alphabetical, LOCTEXT("ProjectSortMode_Alphabetical", "Alphabetical"), FText::GetEmpty());
+	MakeSortMenuEntry(EProjectSortOption::Version, LOCTEXT("ProjectSortMode_Version", "Version"), FText::GetEmpty());
+
+
 	ChildSlot
 	[
 		SNew(SVerticalBox)
@@ -305,22 +351,35 @@ void SProjectBrowser::Construct( const FArguments& InArgs )
 				.HintText(LOCTEXT("FilterHint", "Filter Projects..."))
 				.OnTextChanged(this, &SProjectBrowser::OnFilterTextChanged)
 			]
-			+SHorizontalBox::Slot()
+			+ SHorizontalBox::Slot()
 			.AutoWidth()
 			.VAlign(VAlign_Center)
-			.Padding(FMargin(0, 0, 5.f, 0))
+			.Padding(FMargin(0, 0, 5.f, 5))
 			[
-				SNew(SButton)
-				.ButtonStyle(FAppStyle::Get(), "SimpleButton")
-				.OnClicked(this, &SProjectBrowser::FindProjects)
-				.ToolTipText(LOCTEXT("RefreshProjectList", "Refresh the project list"))
-				.HAlign(HAlign_Center)
-				.VAlign(VAlign_Center)
+				SNew(SComboButton)
+				.ToolTipText(LOCTEXT("SortProjectList", "Sort the project list"))
+				.ComboButtonStyle(FAppStyle::Get(), "SimpleComboButtonWithIcon")
+				.MenuContent()
+				[
+					SortOptionsBuilder.MakeWidget()
+				]
+				.ButtonContent()
 				[
 					SNew(SImage)
-					.Image(FEditorStyle::GetBrush("Icons.Refresh"))
+					.Image(FEditorStyle::GetBrush("Icons.SortDown"))
 					.ColorAndOpacity(FSlateColor::UseForeground())
 				]
+
+			]
+			+SHorizontalBox::Slot()
+			.AutoWidth()
+			.VAlign(VAlign_Center)
+			.Padding(FMargin(0, 0, 5.f, 5))
+			[
+				SNew(SSimpleButton)
+				.OnClicked(this, &SProjectBrowser::FindProjects)
+				.Icon(FEditorStyle::GetBrush("Icons.Refresh"))
+				.ToolTipText(LOCTEXT("RefreshProjectList", "Refresh the project list"))
 			]
 		]
 		+ SVerticalBox::Slot()
@@ -374,7 +433,6 @@ TSharedRef<SToolTip> SProjectBrowser::MakeProjectToolTip( TSharedPtr<FProjectIte
 		AddToToolTipInfoBox( InfoBox, LOCTEXT("ProjectTileTooltipPath", "Path"), FText::FromString(ProjectPath) );
 	}
 
-	if(!ProjectItem->IsUpToDate())
 	{
 		FText Description;
 		if(FDesktopPlatformModule::Get()->IsStockEngineRelease(ProjectItem->EngineIdentifier))
@@ -727,7 +785,7 @@ FReply SProjectBrowser::FindProjects()
 		}
 	}
 
-	ProjectItemsSource.Sort([](auto& A, auto& B) { return A->LastAccessTime < B->LastAccessTime; });
+	SortProjectTiles(CurrentSortOption);
 
 	PopulateFilteredProjects();
 
@@ -1180,6 +1238,32 @@ EVisibility SProjectBrowser::GetNoProjectsAfterFilterErrorVisibility() const
 	return (ProjectItemsSource.Num() && FilteredProjectItemsSource.Num() == 0) ? EVisibility::Visible : EVisibility::Collapsed;
 }
 
+void SProjectBrowser::SortProjectTiles(EProjectSortOption NewSortOption)
+{
+	CurrentSortOption = NewSortOption;
+
+	switch (CurrentSortOption)
+	{
+	case EProjectSortOption::Version:
+		ProjectItemsSource.Sort([](auto& A, auto& B) { return A->CompareEngineVersion(*B); });
+		break;
+	case EProjectSortOption::Alphabetical:
+		ProjectItemsSource.Sort([](auto& A, auto& B) { return A->Name.CompareTo(B->Name) < 0; });
+		break;
+	case EProjectSortOption::LastAccessTime:
+	default:
+		ProjectItemsSource.Sort([](auto& A, auto& B) { return A->LastAccessTime < B->LastAccessTime; });
+		break;
+	}
+
+	PopulateFilteredProjects();
+}
+
+ECheckBoxState SProjectBrowser::GetSortOptionCheckState(EProjectSortOption TestOption) const
+{
+	return CurrentSortOption == TestOption ? ECheckBoxState::Checked : ECheckBoxState::Unchecked;
+}
+
 FText SProjectBrowser:: GetItemHighlightText() const
 {
 	return ProjectItemFilter.GetRawFilterText();
diff --git a/Engine/Source/Editor/GameProjectGeneration/Private/SProjectBrowser.h b/Engine/Source/Editor/GameProjectGeneration/Private/SProjectBrowser.h
--- a/Engine/Source/Editor/GameProjectGeneration/Private/SProjectBrowser.h
+++ b/Engine/Source/Editor/GameProjectGeneration/Private/SProjectBrowser.h
@@ -22,6 +22,19 @@ struct FProjectCategory;
 struct FProjectItem;
 struct FSlateBrush;
 
+enum class EProjectSortOption : uint8
+{
+	/** Sort by the last time the project was opened */
+	LastAccessTime,
+
+	/** Sort by engine version */
+	Version,
+
+	/** By project name */
+	Alphabetical,
+};
+
+
 enum class ECheckBoxState : uint8;
 
 DECLARE_DELEGATE_OneParam(FProjectSelectionChanged, FString);
@@ -94,7 +107,6 @@ protected:
 	/** Handler to check to see if a find in explorer command is allowed */
 	bool CanExecuteFindInExplorer() const;
 
-
 	/** Gets the currently selected template item */
 	TSharedPtr<FProjectItem> GetSelectedProjectItem() const;
 
@@ -141,6 +153,9 @@ private:
 	/** Get the filter text to highlight on items in the list */
 	FText GetItemHighlightText() const;
 
+	void SortProjectTiles(EProjectSortOption NewSortOption);
+	ECheckBoxState GetSortOptionCheckState(EProjectSortOption TestOption) const;
+
 private:
 	TSharedPtr<STileView<TSharedPtr<FProjectItem>>> ProjectTileView;
 	TArray<TSharedPtr<FProjectItem>> ProjectItemsSource;
@@ -164,5 +179,7 @@ private:
 
 	// Holds a delegate that is executed when the new project screen is being requested.
 	FNewProjectScreenRequested NewProjectScreenRequestedDelegate;
+
+	EProjectSortOption CurrentSortOption = EProjectSortOption::LastAccessTime;
 };
 
diff --git a/Engine/Source/Runtime/SlateCore/Private/Styling/StarshipCoreStyle.cpp b/Engine/Source/Runtime/SlateCore/Private/Styling/StarshipCoreStyle.cpp
--- a/Engine/Source/Runtime/SlateCore/Private/Styling/StarshipCoreStyle.cpp
+++ b/Engine/Source/Runtime/SlateCore/Private/Styling/StarshipCoreStyle.cpp
@@ -385,6 +385,9 @@ TSharedRef<ISlateStyle> FStarshipCoreStyle::Create()
 
 		Style->Set("Icons.BulletPoint", new IMAGE_BRUSH_SVG("Starship/Common/bullet-point", Icon8x8));
 
+		Style->Set("Icons.SortDown", new IMAGE_BRUSH_SVG("Starship/Common/SortDown", Icon16x16));
+		Style->Set("Icons.SortUp", new IMAGE_BRUSH_SVG("Starship/Common/SortUp", Icon16x16));
+
 		// Toolbar Size Icons
 		Style->Set("Icons.Toolbar.Settings", new IMAGE_BRUSH_SVG("Starship/Common/settings", Icon20x20));
 	}
