Stable sort project tile entries so they dont move around when file access time is equivalent.

diff --git a/Engine/Source/Editor/GameProjectGeneration/Private/SProjectBrowser.cpp b/Engine/Source/Editor/GameProjectGeneration/Private/SProjectBrowser.cpp
--- a/Engine/Source/Editor/GameProjectGeneration/Private/SProjectBrowser.cpp
+++ b/Engine/Source/Editor/GameProjectGeneration/Private/SProjectBrowser.cpp
@@ -1299,14 +1299,14 @@ void SProjectBrowser::SortProjectTiles(EProjectSortOption NewSortOption)
 	switch (CurrentSortOption)
 	{
 	case EProjectSortOption::Version:
-		ProjectItemsSource.Sort([](auto& A, auto& B) { return A->CompareEngineVersion(*B); });
+		ProjectItemsSource.StableSort([](auto& A, auto& B) { return A->CompareEngineVersion(*B); });
 		break;
 	case EProjectSortOption::Alphabetical:
-		ProjectItemsSource.Sort([](auto& A, auto& B) { return A->Name.CompareTo(B->Name) < 0; });
+		ProjectItemsSource.StableSort([](auto& A, auto& B) { return A->Name.CompareTo(B->Name) < 0; });
 		break;
 	case EProjectSortOption::LastAccessTime:
 	default:
-		ProjectItemsSource.Sort([](auto& A, auto& B) { return A->LastAccessTime < B->LastAccessTime; });
+		ProjectItemsSource.StableSort([](auto& A, auto& B) { return A->LastAccessTime < B->LastAccessTime; });
 		break;
 	}
 
