Fix replacement strings in templates being garbled if the project name contains the template name.

diff --git a/Engine/Source/Editor/GameProjectGeneration/Private/GameProjectUtils.cpp b/Engine/Source/Editor/GameProjectGeneration/Private/GameProjectUtils.cpp
--- a/Engine/Source/Editor/GameProjectGeneration/Private/GameProjectUtils.cpp
+++ b/Engine/Source/Editor/GameProjectGeneration/Private/GameProjectUtils.cpp
@@ -1608,6 +1608,11 @@ static bool SaveConfigValues(const FProjectInformation& InProjectInfo, const TAr
 	return true;
 }
 
+static FString GetReplacePlaceholder(int Idx)
+{
+	return FString::Printf(TEXT("{{{REPLACE:%d}}}"), Idx);
+}
+
 TOptional<FGuid> GameProjectUtils::CreateProjectFromTemplate(const FProjectInformation& InProjectInfo, FText& OutFailReason, FText& OutFailLog, TArray<FString>* OutCreatedFiles)
 {
 	FScopedSlowTask SlowTask(10);
@@ -1781,11 +1786,23 @@ TOptional<FGuid> GameProjectUtils::CreateProjectFromTemplate(const FProjectInfor
 			FString FileContents;
 			if ( FFileHelper::LoadFileToString(FileContents, *FileToFix) )
 			{
-				for ( const FTemplateReplacement& Replacement : TemplateDefs->ReplacementsInFiles )
+				// Substitute strings in two passes to avoid situations where patterns may match the replaced strings.
+				for (int Idx = 0; Idx < TemplateDefs->ReplacementsInFiles.Num(); Idx++)
+				{
+					const FTemplateReplacement& Replacement = TemplateDefs->ReplacementsInFiles[Idx];
+					if (Replacement.Extensions.Contains(FPaths::GetExtension(FileToFix)))
+					{
+						FString Placeholder = GetReplacePlaceholder(Idx);
+						FileContents.ReplaceInline(*Replacement.From, *Placeholder, Replacement.bCaseSensitive ? ESearchCase::CaseSensitive : ESearchCase::IgnoreCase);
+					}
+				}
+				for (int Idx = 0; Idx < TemplateDefs->ReplacementsInFiles.Num(); Idx++)
 				{
-					if ( Replacement.Extensions.Contains( FPaths::GetExtension(FileToFix) ) )
+					const FTemplateReplacement& Replacement = TemplateDefs->ReplacementsInFiles[Idx];
+					if (Replacement.Extensions.Contains(FPaths::GetExtension(FileToFix)))
 					{
-						FileContents = FileContents.Replace(*Replacement.From, *Replacement.To, Replacement.bCaseSensitive ? ESearchCase::CaseSensitive : ESearchCase::IgnoreCase);
+						FString Placeholder = GetReplacePlaceholder(Idx);
+						FileContents.ReplaceInline(*Placeholder, *Replacement.To, ESearchCase::CaseSensitive);
 					}
 				}
 
