New projects always have DX12 by default on Windows

diff --git a/Engine/Source/Editor/GameProjectGeneration/Private/GameProjectUtils.cpp b/Engine/Source/Editor/GameProjectGeneration/Private/GameProjectUtils.cpp
--- a/Engine/Source/Editor/GameProjectGeneration/Private/GameProjectUtils.cpp
+++ b/Engine/Source/Editor/GameProjectGeneration/Private/GameProjectUtils.cpp
@@ -2073,6 +2073,13 @@ void GameProjectUtils::AddHardwareConfigValues(const FProjectInformation& InProj
 			}
 		}
 	}
+
+	// New projects always have DX12 by default on Windows
+	ConfigValues.Emplace(TEXT("DefaultEngine.ini"),
+		TEXT("/Script/WindowsTargetPlatform.WindowsTargetSettings"),
+		TEXT("DefaultGraphicsRHI"),
+		TEXT("DefaultGraphicsRHI_DX12"),
+		false /* ShouldReplaceExistingValue */);
 }
 
 bool GameProjectUtils::GenerateConfigFiles(const FProjectInformation& InProjectInfo, TArray<FString>& OutCreatedFiles, FText& OutFailReason, FGuid& OutProjectID)
