Removing restricted platform names from file

diff --git a/Engine/Source/Editor/GameProjectGeneration/Private/GameProjectUtils.cpp b/Engine/Source/Editor/GameProjectGeneration/Private/GameProjectUtils.cpp
--- a/Engine/Source/Editor/GameProjectGeneration/Private/GameProjectUtils.cpp
+++ b/Engine/Source/Editor/GameProjectGeneration/Private/GameProjectUtils.cpp
@@ -91,71 +91,10 @@
 TWeakPtr<SNotificationItem> GameProjectUtils::UpdateGameProjectNotification = NULL;
 TWeakPtr<SNotificationItem> GameProjectUtils::WarningProjectNameNotification = NULL;
 
-bool GameProjectUtils::bUseAudioMixerForAllPlatforms = false;
-
 constexpr const TCHAR GameProjectUtils::IncludePathFormatString[];
 
-struct FAudioDefaultPlatformSettings
-{
-	FString Name;
-	FAudioPlatformSettings Settings;
-	const TCHAR* ConfigSectionName;
-	bool bUseAudioMixer;
-	FString HardwareOcclusionPlugin;
-	FString HardwareReverbPlugin;
-	FString HardwareSpatializationPlugin;
-
-	FAudioDefaultPlatformSettings(const TCHAR* InConfigSectionName)
-		: ConfigSectionName(InConfigSectionName)
-		, bUseAudioMixer(false)
-	{
-	}
-};
-
 namespace
 {
-	TMap<FString, FAudioDefaultPlatformSettings> GetAudioPlatformProjectDefaultSettings()
-	{
-		TMap<FString, FAudioDefaultPlatformSettings> DefaultProjectSettings;
-
-		// If bUseAudioMixerForAllPlatforms is set to false, uncomment the following line to enable
-		// the new audio mixer on specific platform. Ex. for Windows:
-		// WindowsSettings.bUseAudioMixer = true;
-
-		FAudioDefaultPlatformSettings AndroidSettings(TEXT("/Script/AndroidRuntimeSettings.AndroidRuntimeSettings"));
-		AndroidSettings.Settings.MaxChannels = 12;
-		DefaultProjectSettings.Add(TEXT("Android"), AndroidSettings);
-
-		FAudioDefaultPlatformSettings IOSSettings(TEXT("/Script/IOSRuntimeSettings.IOSRuntimeSettings"));
-		IOSSettings.Settings.MaxChannels = 16;
-		DefaultProjectSettings.Add(TEXT("IOS"), IOSSettings);
-
-		FAudioDefaultPlatformSettings LinuxSettings(TEXT("/Script/LinuxTargetPlatform.LinuxTargetSettings"));
-		LinuxSettings.Settings.MaxChannels = 16;
-		DefaultProjectSettings.Add(TEXT("Linux"), LinuxSettings);
-
-		FAudioDefaultPlatformSettings MacSettings(TEXT("/Script/MacTargetPlatform.MacTargetSettings"));
-		DefaultProjectSettings.Add(TEXT("Mac"), MacSettings);
-
-		FAudioDefaultPlatformSettings PS4Settings(TEXT("/Script/PS4PlatformEditor.PS4TargetSettings"));
-		PS4Settings.Settings.CallbackBufferFrameSize = 256;
-		PS4Settings.Settings.NumBuffers = 7;
-		PS4Settings.Settings.NumSourceWorkers = 4;
-		PS4Settings.HardwareSpatializationPlugin = TEXT("Audio3D");
-		DefaultProjectSettings.Add(TEXT("PS4"), PS4Settings);
-
-		FAudioDefaultPlatformSettings SwitchSettings(TEXT("/Script/SwitchRuntimeSettings.SwitchRuntimeSettings"));
-		SwitchSettings.Settings.MaxChannels = 16;
-		DefaultProjectSettings.Add(TEXT("Switch"), SwitchSettings);
-
-		FAudioDefaultPlatformSettings WindowsSettings(TEXT("/Script/WindowsTargetPlatform.WindowsTargetSettings"));
-		WindowsSettings.Settings.CallbackBufferFrameSize = 256;
-		WindowsSettings.Settings.NumBuffers = 7;
-		DefaultProjectSettings.Add(TEXT("Windows"), WindowsSettings);
-
-		return MoveTemp(DefaultProjectSettings);
-	}
-
 	// @todo: This is currently not called from anywhere as this approach does not work for binary builds.
 	/** Set the state of XR plugins in OutProject based on the flags in InProjectInfo. */
 	void SetXRPluginStates(const FProjectInformation& InProjectInfo, FProjectDescriptor& OutProject)
@@ -2094,12 +2033,9 @@ bool GameProjectUtils::GenerateConfigFiles(const FProjectInformation& InProjectI
 		const FString DefaultEngineIniFilename = ProjectConfigPath / TEXT("DefaultEngine.ini");
 		FString FileContents;
 
-		if(bUseAudioMixerForAllPlatforms)
-		{
-			FileContents += LINE_TERMINATOR;
-			FileContents += TEXT("[Audio]") LINE_TERMINATOR;
-			FileContents += TEXT("UseAudioMixer=True") LINE_TERMINATOR;
-		}
+		FileContents += LINE_TERMINATOR;
+		FileContents += TEXT("[Audio]") LINE_TERMINATOR;
+		FileContents += TEXT("UseAudioMixer=True") LINE_TERMINATOR;
 
 		if (InProjectInfo.bForceExtendedLuminanceRange)
 		{
@@ -2202,71 +2138,6 @@ bool GameProjectUtils::GenerateConfigFiles(const FProjectInformation& InProjectI
 
 bool GameProjectUtils::GeneratePlatformConfigFiles(const FProjectInformation& InProjectInfo, FText& OutFailReason)
 {
-	TMap<FString, FAudioDefaultPlatformSettings> ProjectDefaults = GetAudioPlatformProjectDefaultSettings();
-
-	static const FAudioPlatformSettings DefaultSettings;
-
-	for (TPair<FString, FAudioDefaultPlatformSettings>& SettingsPair : ProjectDefaults)
-	{
-		FString FileContents;
-
-		if (bUseAudioMixerForAllPlatforms || SettingsPair.Value.bUseAudioMixer)
-		{
-			FileContents += TEXT("[Audio]") LINE_TERMINATOR;
-			FileContents += TEXT("UseAudioMixer=True") LINE_TERMINATOR;
-			FileContents += LINE_TERMINATOR;
-		}
-
-		const FString& PlatformName = SettingsPair.Key;
-		const FAudioPlatformSettings& PlatformSettings = SettingsPair.Value.Settings;
-
-		FileContents += TEXT("[") + FString(SettingsPair.Value.ConfigSectionName) + TEXT("]") + LINE_TERMINATOR;
-
-		if (DefaultSettings.SampleRate == PlatformSettings.SampleRate)
-		{
-			FileContents += TEXT(";");
-		}
-		FileContents += TEXT("AudioSampleRate=") + FString::Printf(TEXT("%d"), PlatformSettings.SampleRate) + LINE_TERMINATOR;
-
-		if (DefaultSettings.MaxChannels == PlatformSettings.MaxChannels)
-		{
-			FileContents += TEXT(";");
-		}
-		FileContents += TEXT("AudioMaxChannels=") + FString::Printf(TEXT("%d"), PlatformSettings.MaxChannels) + LINE_TERMINATOR;
-
-		if (DefaultSettings.CallbackBufferFrameSize == PlatformSettings.CallbackBufferFrameSize)
-		{
-			FileContents += TEXT(";");
-		}
-		FileContents += TEXT("AudioCallbackBufferFrameSize=") + FString::Printf(TEXT("%d"), PlatformSettings.CallbackBufferFrameSize) + LINE_TERMINATOR;
-
-		if (DefaultSettings.NumBuffers == PlatformSettings.NumBuffers)
-		{
-			FileContents += TEXT(";");
-		}
-		FileContents += TEXT("AudioNumBuffersToEnqueue=") + FString::Printf(TEXT("%d"), PlatformSettings.NumBuffers) + LINE_TERMINATOR;
-
-		if (DefaultSettings.NumSourceWorkers == PlatformSettings.NumSourceWorkers)
-		{
-			FileContents += TEXT(";");
-		}
-		FileContents += TEXT("AudioNumSourceWorkers=") + FString::Printf(TEXT("%d"), PlatformSettings.NumSourceWorkers) + LINE_TERMINATOR;
-
-		FileContents += LINE_TERMINATOR;
-		FileContents += TEXT("; Audio Plugins (must be enabled in .uproject") LINE_TERMINATOR;
-		FileContents += TEXT("; ReverbPlugin=") + SettingsPair.Value.HardwareReverbPlugin + LINE_TERMINATOR;
-		FileContents += TEXT("; OcclusionPlugin=") + SettingsPair.Value.HardwareOcclusionPlugin + LINE_TERMINATOR;
-		FileContents += TEXT("; SpatializationPlugin=") + SettingsPair.Value.HardwareSpatializationPlugin + LINE_TERMINATOR;
-
-		const FString NewProjectFolder = FPaths::GetPath(InProjectInfo.ProjectFilename);
-		const FString ProjectConfigPath = NewProjectFolder / TEXT("Platforms") / PlatformName / TEXT("Config");
-		const FString PlatformEngineIniFilename = ProjectConfigPath / PlatformName + TEXT("Engine.ini");
-		if (!WriteOutputFile(PlatformEngineIniFilename, FileContents, OutFailReason))
-		{
-			return false;
-		}
-	}
-
 	return true;
 }
 
diff --git a/Engine/Source/Runtime/Android/AudioMixerAndroid/Private/AudioMixerPlatformAndroid.h b/Engine/Source/Runtime/Android/AudioMixerAndroid/Private/AudioMixerPlatformAndroid.h
--- a/Engine/Source/Runtime/Android/AudioMixerAndroid/Private/AudioMixerPlatformAndroid.h
+++ b/Engine/Source/Runtime/Android/AudioMixerAndroid/Private/AudioMixerPlatformAndroid.h
@@ -20,7 +20,7 @@ namespace Audio
 		~FMixerPlatformAndroid();
 
 		//~ Begin IAudioMixerPlatformInterface
-		virtual EAudioMixerPlatformApi::Type GetPlatformApi() const override { return EAudioMixerPlatformApi::OpenSLES; }
+		virtual FString GetPlatformApi() const override { return TEXT("OpenSLES"); }
 		virtual bool InitializeHardware() override;
 		virtual bool TeardownHardware() override;
 		virtual bool IsInitialized() const override;
diff --git a/Engine/Source/Runtime/Apple/AudioMixerAudioUnit/Private/AudioMixerPlatformAudioUnit.h b/Engine/Source/Runtime/Apple/AudioMixerAudioUnit/Private/AudioMixerPlatformAudioUnit.h
--- a/Engine/Source/Runtime/Apple/AudioMixerAudioUnit/Private/AudioMixerPlatformAudioUnit.h
+++ b/Engine/Source/Runtime/Apple/AudioMixerAudioUnit/Private/AudioMixerPlatformAudioUnit.h
@@ -20,7 +20,7 @@ namespace Audio
 		~FMixerPlatformAudioUnit();
 		
 		//~ Begin IAudioMixerPlatformInterface
-		virtual EAudioMixerPlatformApi::Type GetPlatformApi() const override { return EAudioMixerPlatformApi::AudioUnit; }
+		virtual FString GetPlatformApi() const override { return TEXT("AudioUnit"); }
 		virtual bool InitializeHardware() override;
 		virtual bool CheckAudioDeviceChange() override;
 		virtual bool TeardownHardware() override;
diff --git a/Engine/Source/Runtime/Apple/AudioMixerCoreAudio/Private/AudioMixerPlatformCoreAudio.h b/Engine/Source/Runtime/Apple/AudioMixerCoreAudio/Private/AudioMixerPlatformCoreAudio.h
--- a/Engine/Source/Runtime/Apple/AudioMixerCoreAudio/Private/AudioMixerPlatformCoreAudio.h
+++ b/Engine/Source/Runtime/Apple/AudioMixerCoreAudio/Private/AudioMixerPlatformCoreAudio.h
@@ -20,7 +20,7 @@ namespace Audio
 		~FMixerPlatformCoreAudio();
 
 		//~ Begin IAudioMixerPlatformInterface
-		virtual EAudioMixerPlatformApi::Type GetPlatformApi() const override { return EAudioMixerPlatformApi::CoreAudio; }
+		virtual FString GetPlatformApi() const override { return TEXT("CoreAudio"); }
 		virtual bool InitializeHardware() override;
 		virtual bool CheckAudioDeviceChange() override;
 		virtual bool TeardownHardware() override;
diff --git a/Engine/Source/Runtime/AudioMixer/Private/AudioMixerDevice.cpp b/Engine/Source/Runtime/AudioMixer/Private/AudioMixerDevice.cpp
--- a/Engine/Source/Runtime/AudioMixer/Private/AudioMixerDevice.cpp
+++ b/Engine/Source/Runtime/AudioMixer/Private/AudioMixerDevice.cpp
@@ -344,11 +344,11 @@ namespace Audio
 	
 		LLM_SCOPE(ELLMTag::AudioMixer);
 
-		// Log that we're inside the audio mixer
-		UE_LOG(LogAudioMixer, Display, TEXT("Initializing audio mixer."));
 
 		if (AudioMixerPlatform && AudioMixerPlatform->InitializeHardware())
 		{
+			UE_LOG(LogAudioMixer, Display, TEXT("Initializing audio mixer using platform API: '%s'"), *AudioMixerPlatform->GetPlatformApi());
+
 			const UAudioSettings* AudioSettings = GetDefault<UAudioSettings>();
 			MonoChannelUpmixMethod = AudioSettings->MonoChannelUpmixMethod;
 			PanningMethod = AudioSettings->PanningMethod;
@@ -404,7 +404,7 @@ namespace Audio
 			{
 				// Get the platform device info we're using
 				PlatformInfo = AudioMixerPlatform->GetPlatformDeviceInfo();
-				UE_LOG(LogAudioMixer, Display, TEXT("Using Audio Device %s"), *PlatformInfo.Name);
+				UE_LOG(LogAudioMixer, Display, TEXT("Using Audio Hardware Device %s"), *PlatformInfo.Name);
 
 				// Initialize some data that depends on speaker configuration, etc.
 				InitializeChannelAzimuthMap(PlatformInfo.NumChannels);
@@ -467,7 +467,7 @@ namespace Audio
 		}
 		else if (AudioMixerPlatform)
 		{
-			UE_LOG(LogAudioMixer, Warning, TEXT("Failed to Open audio stream"));
+			UE_LOG(LogAudioMixer, Warning, TEXT("Failed to initialize audio mixer for platform API: '%s'"), *AudioMixerPlatform->GetPlatformApi());
 		}
 
 		return false;
diff --git a/Engine/Source/Runtime/AudioMixerCore/Public/AudioMixer.h b/Engine/Source/Runtime/AudioMixerCore/Public/AudioMixer.h
--- a/Engine/Source/Runtime/AudioMixerCore/Public/AudioMixer.h
+++ b/Engine/Source/Runtime/AudioMixerCore/Public/AudioMixer.h
@@ -417,8 +417,8 @@ namespace Audio
 		/** Virtual destructor. */
 		virtual ~IAudioMixerPlatformInterface();
 
-		/** Returns the platform API enumeration. */
-		virtual EAudioMixerPlatformApi::Type GetPlatformApi() const = 0;
+		/** Returns the platform API name. */
+		virtual FString GetPlatformApi() const = 0;
 
 		/** Initialize the hardware. */
 		virtual bool InitializeHardware() = 0;
diff --git a/Engine/Source/Runtime/AudioMixerCore/Public/AudioMixerTypes.h b/Engine/Source/Runtime/AudioMixerCore/Public/AudioMixerTypes.h
--- a/Engine/Source/Runtime/AudioMixerCore/Public/AudioMixerTypes.h
+++ b/Engine/Source/Runtime/AudioMixerCore/Public/AudioMixerTypes.h
@@ -6,22 +6,6 @@
 
 namespace Audio {
 
-	namespace EAudioMixerPlatformApi
-	{
-		enum Type
-		{
-			XAudio2, 	// Windows, XBoxOne
-			AudioOut, 	// PS4
-			CoreAudio, 	// Mac
-			AudioUnit, 	// iOS
-			SDL2,		// Linux
-			OpenSLES, 	// Android
-			Switch, 	// Switch
-			Other,      // Generic output type.
-			Null		// Unknown/not Supported
-		};
-	}
-
 	namespace EAudioMixerStreamDataFormat
 	{
 		enum Type
@@ -126,7 +110,7 @@ struct AUDIOMIXERCORE_API FAudioPlatformSettings
 		: SampleRate(48000)
 		, CallbackBufferFrameSize(1024)
 		, NumBuffers(2)
-		, MaxChannels(32)
+		, MaxChannels(0) // This needs to be 0 to indicate it's not overridden from the audio settings object, which is the default used on all platforms
 		, NumSourceWorkers(0)
 	{
 	}
diff --git a/Engine/Source/Runtime/Linux/AudioMixerSDL/Private/AudioMixerPlatformSDL.h b/Engine/Source/Runtime/Linux/AudioMixerSDL/Private/AudioMixerPlatformSDL.h
--- a/Engine/Source/Runtime/Linux/AudioMixerSDL/Private/AudioMixerPlatformSDL.h
+++ b/Engine/Source/Runtime/Linux/AudioMixerSDL/Private/AudioMixerPlatformSDL.h
@@ -21,7 +21,7 @@ namespace Audio
 		virtual ~FMixerPlatformSDL();
 
 		//~ Begin IAudioMixerPlatformInterface Interface
-		EAudioMixerPlatformApi::Type GetPlatformApi() const override { return EAudioMixerPlatformApi::SDL2; }
+		FString GetPlatformApi() const override { return TEXT("SDL2"); }
 		bool InitializeHardware() override;
 		bool TeardownHardware() override;
 		bool IsInitialized() const override;
diff --git a/Engine/Source/Runtime/NonRealtimeAudioRenderer/Public/AudioMixerPlatformNonRealtime.h b/Engine/Source/Runtime/NonRealtimeAudioRenderer/Public/AudioMixerPlatformNonRealtime.h
--- a/Engine/Source/Runtime/NonRealtimeAudioRenderer/Public/AudioMixerPlatformNonRealtime.h
+++ b/Engine/Source/Runtime/NonRealtimeAudioRenderer/Public/AudioMixerPlatformNonRealtime.h
@@ -21,7 +21,7 @@ namespace Audio
 		void CloseFile();
 
 		//~ Begin IAudioMixerPlatformInterface
-		virtual EAudioMixerPlatformApi::Type GetPlatformApi() const override { return EAudioMixerPlatformApi::XAudio2; }
+		virtual FString GetPlatformApi() const override { return TEXT("NonRealtime"); }
 		virtual bool InitializeHardware() override;
 		virtual bool CheckAudioDeviceChange() override;
 		virtual bool TeardownHardware() override;
diff --git a/Engine/Source/Runtime/Windows/AudioMixerXAudio2/Private/AudioMixerPlatformXAudio2.h b/Engine/Source/Runtime/Windows/AudioMixerXAudio2/Private/AudioMixerPlatformXAudio2.h
--- a/Engine/Source/Runtime/Windows/AudioMixerXAudio2/Private/AudioMixerPlatformXAudio2.h
+++ b/Engine/Source/Runtime/Windows/AudioMixerXAudio2/Private/AudioMixerPlatformXAudio2.h
@@ -63,7 +63,7 @@ namespace Audio
 		~FMixerPlatformXAudio2();
 
 		//~ Begin IAudioMixerPlatformInterface
-		virtual EAudioMixerPlatformApi::Type GetPlatformApi() const override { return EAudioMixerPlatformApi::XAudio2; }
+		virtual FString GetPlatformApi() const override { return TEXT("XAudio2"); }
 		virtual bool InitializeHardware() override;
 		virtual bool CheckAudioDeviceChange() override;
 		virtual bool TeardownHardware() override;
