Editor: explicitly enable Modeling Mode in code-generated blank template projects.
#rnx

diff --git a/Engine/Source/Editor/GameProjectGeneration/Private/GameProjectUtils.cpp b/Engine/Source/Editor/GameProjectGeneration/Private/GameProjectUtils.cpp
--- a/Engine/Source/Editor/GameProjectGeneration/Private/GameProjectUtils.cpp
+++ b/Engine/Source/Editor/GameProjectGeneration/Private/GameProjectUtils.cpp
@@ -1404,6 +1404,23 @@ TOptional<FGuid> GameProjectUtils::GenerateProjectFromScratch(const FProjectInfo
 		{
 			Project.Modules.Add(FModuleDescriptor(*StartupModuleNames[Idx]));
 		}
+		
+		//=====================================================================
+		// Explicitly enable Modeling Mode plugin in Blank Template,
+		// with AllowList=Editor flags. In 5.0 the Modeling Mode plugin
+		// cannot be enabledByDefault in the .uplugin file due to 
+		// dependent Runtime modules that should not be included in all
+		// game builds. So, In 5.0 the plugin is explicitly enabled here
+		// for Blank projects. The uplugin-level issue is expected to 
+		// be resolved in 5.1, at which point this code block will be deleted
+		//=====================================================================
+		TSharedPtr<IPlugin> ModelingModePlugin = IPluginManager::Get().FindPlugin(TEXT("ModelingToolsEditorMode"));
+		if ( ModelingModePlugin.IsValid() )
+		{
+			FPluginReferenceDescriptor ModelingModeDescriptor(ModelingModePlugin->GetName(), true);
+			ModelingModeDescriptor.TargetAllowList.Add(EBuildTargetType::Editor);
+			Project.Plugins.Add(ModelingModeDescriptor);
+		}
 
 		Project.bIsEnterpriseProject = InProjectInfo.bIsEnterpriseProject;
 
