Updated error messages when no compiler/IDE was found when opening or creating C++ projects/files

diff --git a/Engine/Source/Editor/GameProjectGeneration/Private/SNewClassDialog.cpp b/Engine/Source/Editor/GameProjectGeneration/Private/SNewClassDialog.cpp
--- a/Engine/Source/Editor/GameProjectGeneration/Private/SNewClassDialog.cpp
+++ b/Engine/Source/Editor/GameProjectGeneration/Private/SNewClassDialog.cpp
@@ -989,7 +989,8 @@ FText SNewClassDialog::GetGlobalErrorLabelText() const
 	if ( ClassDomain == EClassDomain::Native && !FSourceCodeNavigation::IsCompilerAvailable() )
 	{
 #if PLATFORM_LINUX
-		return FText::Format(LOCTEXT("NoCompilerFoundNewClassLinux", "In order to use a C++ template, you must first install or disable {0}."), FSourceCodeNavigation::GetSuggestedSourceCodeIDE());
+		return FText::Format(LOCTEXT("NoCompilerFoundNewClassLinux", "Your IDE {0} is missing or incorrectly configured, please consider using {1}"),
+			FSourceCodeNavigation::GetSelectedSourceCodeIDE(), FSourceCodeNavigation::GetSuggestedSourceCodeIDE());
 #else
 		return FText::Format(LOCTEXT("NoCompilerFoundNewClass", "No compiler was found. In order to use C++ code, you must first install {0}."), FSourceCodeNavigation::GetSuggestedSourceCodeIDE());
 #endif
diff --git a/Engine/Source/Editor/GameProjectGeneration/Private/SProjectDialog.cpp b/Engine/Source/Editor/GameProjectGeneration/Private/SProjectDialog.cpp
--- a/Engine/Source/Editor/GameProjectGeneration/Private/SProjectDialog.cpp
+++ b/Engine/Source/Editor/GameProjectGeneration/Private/SProjectDialog.cpp
@@ -2043,7 +2043,8 @@ void SProjectDialog::UpdateProjectFileValidity()
 					}
 					else
 					{
-						LastGlobalValidityErrorText = FText::Format(LOCTEXT("NoCompilerFound", "In order to use a C++ template, you must first install or disable {0}."), FSourceCodeNavigation::GetSuggestedSourceCodeIDE());
+						LastGlobalValidityErrorText = FText::Format(LOCTEXT("NoCompilerFoundProjectDialog", "Your IDE {0} is missing or incorrectly configured, please consider using {1}"),
+			FSourceCodeNavigation::GetSelectedSourceCodeIDE(), FSourceCodeNavigation::GetSuggestedSourceCodeIDE());
 					}
 				}
 				else if (!FDesktopPlatformModule::Get()->IsUnrealBuildToolAvailable())
