Prevent starter content from overriding default maps with table and chairs map. We no longer want that to be the default if you include starter content in a new project

diff --git a/Engine/Source/Editor/GameProjectGeneration/Private/GameProjectUtils.cpp b/Engine/Source/Editor/GameProjectGeneration/Private/GameProjectUtils.cpp
--- a/Engine/Source/Editor/GameProjectGeneration/Private/GameProjectUtils.cpp
+++ b/Engine/Source/Editor/GameProjectGeneration/Private/GameProjectUtils.cpp
@@ -1788,7 +1788,6 @@ TOptional<FGuid> GameProjectUtils::CreateProjectFromTemplate(const FProjectInfor
 	AddLumenConfigValues(InProjectInfo, ConfigValuesToSet);
 	AddRaytracingConfigValues(InProjectInfo, ConfigValuesToSet);
 	AddNewProjectDefaultShadowConfigValues(InProjectInfo, ConfigValuesToSet);
-	AddDefaultMapConfigValues(InProjectInfo, ConfigValuesToSet);
 	AddWorldPartitionConfigValues(InProjectInfo, ConfigValuesToSet);
 	
 	TemplateDefs->AddConfigValues(ConfigValuesToSet, TemplateName, ProjectName, InProjectInfo.bShouldGenerateCode);
