Fix for missing MobileStarterContent.upack but the import checkbox is enabled, resulting import error

diff --git a/Engine/Source/Editor/GameProjectGeneration/Private/GameProjectUtils.cpp b/Engine/Source/Editor/GameProjectGeneration/Private/GameProjectUtils.cpp
--- a/Engine/Source/Editor/GameProjectGeneration/Private/GameProjectUtils.cpp
+++ b/Engine/Source/Editor/GameProjectGeneration/Private/GameProjectUtils.cpp
@@ -185,11 +185,10 @@ namespace
 	{
 		if (InProjectInfo.bIsBlankTemplate &&
 			InProjectInfo.bCopyStarterContent &&
-			GameProjectUtils::IsStarterContentAvailableForNewProjects())
+			GameProjectUtils::IsUsingEngineStarterContent(InProjectInfo) &&
+			GameProjectUtils::IsEngineStarterContentAvailable() )
 		{
-			const FString DefaultMap = InProjectInfo.TargetedHardware == EHardwareClass::Mobile ?
-				TEXT("/Game/MobileStarterContent/Maps/Minimal_Default") :
-				TEXT("/Game/StarterContent/Maps/Minimal_Default");
+			const FString DefaultMap = TEXT("/Game/StarterContent/Maps/Minimal_Default");
 
 			ConfigValues.Emplace(TEXT("DefaultEngine.ini"),
 				TEXT("/Script/EngineSettings.GameMapsSettings"),
@@ -766,13 +765,17 @@ bool GameProjectUtils::OpenCodeIDE(const FString& ProjectFile, FText& OutFailRea
 	return true;
 }
 
-bool GameProjectUtils::IsStarterContentAvailableForNewProjects()
+bool GameProjectUtils::IsEngineStarterContentAvailable()
 {
 	TArray<FString> OutFilenames;
 	IFileManager::Get().FindFilesRecursive(OutFilenames, *FPaths::FeaturePackDir(), TEXT("*StarterContent.upack"), /*Files=*/true, /*Directories=*/false);
 	return OutFilenames.Num() > 0;
 }
 
+bool GameProjectUtils::IsUsingEngineStarterContent(const FProjectInformation& InProjectInfo)
+{
+	return InProjectInfo.StarterContent.IsEmpty();
+}
 
 FString GameProjectUtils::GetStarterContentName(const FProjectInformation& InProjectInfo)
 {
@@ -781,11 +784,6 @@ FString GameProjectUtils::GetStarterContentName(const FProjectInformation& InPro
 		return InProjectInfo.StarterContent;
 	}
 
-	if (InProjectInfo.TargetedHardware == EHardwareClass::Mobile)
-	{
-		return TEXT("MobileStarterContent");
-	}
-
 	return TEXT("StarterContent");
 }
 
@@ -2065,18 +2063,10 @@ bool GameProjectUtils::GenerateConfigFiles(const FProjectInformation& InProjectI
 			FileContents += LINE_TERMINATOR;
 			FileContents += TEXT("[/Script/EngineSettings.GameMapsSettings]") LINE_TERMINATOR;
 
-			if (GameProjectUtils::IsStarterContentAvailableForNewProjects())
+			if (GameProjectUtils::IsEngineStarterContentAvailable() && GameProjectUtils::IsUsingEngineStarterContent(InProjectInfo)) // if use Engine StarterContent
 			{
-				if (InProjectInfo.TargetedHardware == EHardwareClass::Mobile)
-				{
-					FileContents += TEXT("EditorStartupMap=/Game/MobileStarterContent/Maps/Minimal_Default") LINE_TERMINATOR;
-					FileContents += TEXT("GameDefaultMap=/Game/MobileStarterContent/Maps/Minimal_Default") LINE_TERMINATOR;
-				}
-				else
-				{
-					FileContents += TEXT("EditorStartupMap=/Game/StarterContent/Maps/Minimal_Default") LINE_TERMINATOR;
-					FileContents += TEXT("GameDefaultMap=/Game/StarterContent/Maps/Minimal_Default") LINE_TERMINATOR;
-				}
+				FileContents += TEXT("EditorStartupMap=/Game/StarterContent/Maps/Minimal_Default") LINE_TERMINATOR;
+				FileContents += TEXT("GameDefaultMap=/Game/StarterContent/Maps/Minimal_Default") LINE_TERMINATOR;
 			}
 
 			if (InProjectInfo.bShouldGenerateCode)
diff --git a/Engine/Source/Editor/GameProjectGeneration/Private/SProjectDialog.cpp b/Engine/Source/Editor/GameProjectGeneration/Private/SProjectDialog.cpp
--- a/Engine/Source/Editor/GameProjectGeneration/Private/SProjectDialog.cpp
+++ b/Engine/Source/Editor/GameProjectGeneration/Private/SProjectDialog.cpp
@@ -787,6 +787,7 @@ TSharedRef<SWidget> SProjectDialog::MakeHybridView(EProjectDialogModeMode Mode)
 
 	SelectedHardwareClassTarget = EHardwareClass::Desktop;
 	SelectedGraphicsPreset = EGraphicsPreset::Maximum;
+	bIsStarterContentAvailable = true;
 
 	// Find all template projects
 	Templates = FindTemplateProjects();
@@ -1011,15 +1012,15 @@ TSharedRef<SWidget> SProjectDialog::MakeProjectOptionsWidget()
 	if (!HiddenSettings.Contains(ETemplateSetting::StarterContent))
 	{
 		FProjectInformation ProjectInfo = CreateProjectInfo();
+		bIsStarterContentAvailable = GameProjectUtils::IsStarterContentAvailableForProject(ProjectInfo);
 
-		const bool bIsStarterContentAvailable = GameProjectUtils::IsStarterContentAvailableForProject(ProjectInfo);
 		ProjectOptionsBox->AddSlot()
 			.Padding(0.0f, 8.0f, 0.0f, 0.0f)
 			.AutoHeight()
 			[
 				SNew(SHorizontalBox)
 				.ToolTipText(LOCTEXT("CopyStarterContent_ToolTip", "Enable to include an additional content pack containing simple placeable meshes with basic materials and textures.\nYou can also add the Starter Content to your project later using the Content Browser."))
-				.IsEnabled(bIsStarterContentAvailable)
+				.IsEnabled(this, &SProjectDialog::IsStarterContentAvailable)
 				+ SHorizontalBox::Slot()
 				.Padding(0.0f, 0.0f, 8.0f, 0.0f)
 				.VAlign(VAlign_Center)
@@ -1259,6 +1260,9 @@ void SProjectDialog::OnSetBlueprintOrCppIndex(int32 Index)
 {
 	bShouldGenerateCode = Index == 1;
 	UpdateProjectFileValidity();
+	FProjectInformation ProjectInfo = CreateProjectInfo();
+	bIsStarterContentAvailable = GameProjectUtils::IsStarterContentAvailableForProject(ProjectInfo);
+	bCopyStarterContent &= bIsStarterContentAvailable;
 }
 
 void SProjectDialog::SetHardwareClassTarget(EHardwareClass InHardwareClass)
@@ -1631,21 +1635,18 @@ FProjectInformation SProjectDialog::CreateProjectInfo() const
 	ProjectInfo.bIsBlankTemplate = SelectedTemplate->bIsBlankTemplate;
 	ProjectInfo.bForceExtendedLuminanceRange = SelectedTemplate->bIsBlankTemplate;
 
-	if (bCopyStarterContent)
+	if (bShouldGenerateCode)
 	{
-		if (bShouldGenerateCode)
+		if (SelectedTemplate->CodeTemplateDefs != nullptr)
 		{
-			if (SelectedTemplate->CodeTemplateDefs != nullptr)
-			{
-				ProjectInfo.StarterContent = SelectedTemplate->CodeTemplateDefs->StarterContent;
-			}
+			ProjectInfo.StarterContent = SelectedTemplate->CodeTemplateDefs->StarterContent;
 		}
-		else
+	}
+	else
+	{
+		if (SelectedTemplate->BlueprintTemplateDefs != nullptr)
 		{
-			if (SelectedTemplate->BlueprintTemplateDefs != nullptr)
-			{
-				ProjectInfo.StarterContent = SelectedTemplate->BlueprintTemplateDefs->StarterContent;
-			}
+			ProjectInfo.StarterContent = SelectedTemplate->BlueprintTemplateDefs->StarterContent;
 		}
 	}
 
diff --git a/Engine/Source/Editor/GameProjectGeneration/Private/SProjectDialog.h b/Engine/Source/Editor/GameProjectGeneration/Private/SProjectDialog.h
--- a/Engine/Source/Editor/GameProjectGeneration/Private/SProjectDialog.h
+++ b/Engine/Source/Editor/GameProjectGeneration/Private/SProjectDialog.h
@@ -146,6 +146,9 @@ private:
 	FReply OnCancel() const;
 
 	ECheckBoxState OnGetRaytracingEnabledCheckState() const { return bEnableRaytracing ? ECheckBoxState::Checked : ECheckBoxState::Unchecked; }
+	
+	bool IsStarterContentAvailable() const {return bIsStarterContentAvailable; }
+
 	void OnSetRaytracingEnabled(ECheckBoxState NewState);
 
 	int32 OnGetBlueprintOrCppIndex() const { return bShouldGenerateCode ? 1 : 0; }
@@ -224,6 +227,9 @@ private:
 	/** True if user has selected to copy starter content. */
 	bool bCopyStarterContent;
 
+	/** True if starter content is available for current project. */
+	bool bIsStarterContentAvailable;
+
 	/** Whether or not to enable XR in the created project. */
 	bool bEnableXR;
 
diff --git a/Engine/Source/Editor/GameProjectGeneration/Public/GameProjectUtils.h b/Engine/Source/Editor/GameProjectGeneration/Public/GameProjectUtils.h
--- a/Engine/Source/Editor/GameProjectGeneration/Public/GameProjectUtils.h
+++ b/Engine/Source/Editor/GameProjectGeneration/Public/GameProjectUtils.h
@@ -176,8 +176,13 @@ public:
 	/** Creates code project files for a new game project. On failure, OutFailReason and OutFailLog will be populated. */
 	static bool GenerateCodeProjectFiles(const FString& ProjectFilename, FText& OutFailReason, FText& OutFailLog);
 
-	/** Returns true if there are starter content files available for instancing into new projects. */
-	static bool IsStarterContentAvailableForNewProjects();
+	/** Returns true if there are engine starter content files available for instancing into new projects. */
+	static bool IsEngineStarterContentAvailable();
+
+	/** Returns true if the given project is referencing engine starter content pack. */
+	static bool IsUsingEngineStarterContent(const FProjectInformation& ProjectInfo);
+
+	/** Returns true if there are any starter content packs are available for given project. */
 	static bool IsStarterContentAvailableForProject(const FProjectInformation& ProjectInfo);
 
 	/**
