Disallow creating projects with the same names as platform modules

Those projects cannot be loaded due to the modules conflicting.

diff --git a/Engine/Source/Editor/GameProjectGeneration/Private/GameProjectUtils.cpp b/Engine/Source/Editor/GameProjectGeneration/Private/GameProjectUtils.cpp
--- a/Engine/Source/Editor/GameProjectGeneration/Private/GameProjectUtils.cpp
+++ b/Engine/Source/Editor/GameProjectGeneration/Private/GameProjectUtils.cpp
@@ -630,6 +630,14 @@ bool GameProjectUtils::IsValidProjectFileForCreation(const FString& ProjectFile,
 		return false;
 	}
 
+	if (NameMatchesPlatformModuleName(BaseProjectFile))
+	{
+		FFormatNamedArguments Args;
+		Args.Add(TEXT("PlatformModuleName"), FText::FromString(BaseProjectFile));
+		OutFailReason = FText::Format(LOCTEXT("ProjectNameConflictsWithPlatformModuleName", "Project name conflicts with a platform name: {PlatformModuleName}"), Args);
+		return false;
+	}
+
 	if ( !FPaths::ValidatePath(FPaths::GetPath(ProjectFile), &OutFailReason) )
 	{
 		return false;
@@ -1948,6 +1956,28 @@ bool GameProjectUtils::NameContainsOnlyLegalCharacters(const FString& TestName,
 	return !bContainsIllegalCharacters;
 }
 
+bool GameProjectUtils::NameMatchesPlatformModuleName(const FString& TestName)
+{
+	for (auto Pair : FDataDrivenPlatformInfoRegistry::GetAllPlatformInfos())
+	{
+		FString PlatformNameString = Pair.Key.ToString();
+		if ((PlatformNameString == TestName) || ((PlatformNameString += "TargetPlatform") == TestName))
+		{
+			return true;
+		}
+	}
+	TArray<FString> CustomTargetPlatformModules;
+	GConfig->GetArray(TEXT("CustomTargetPlatforms"), TEXT("ModuleName"), CustomTargetPlatformModules, GEditorIni);
+	for (const FString& ModuleName : CustomTargetPlatformModules)
+	{
+		if (TestName == ModuleName)
+		{
+			return true;
+		}
+	}
+	return false;
+}
+
 bool GameProjectUtils::ProjectFileExists(const FString& ProjectFile)
 {
 	return FPlatformFileManager::Get().GetPlatformFile().FileExists(*ProjectFile);
diff --git a/Engine/Source/Editor/GameProjectGeneration/Public/GameProjectUtils.h b/Engine/Source/Editor/GameProjectGeneration/Public/GameProjectUtils.h
--- a/Engine/Source/Editor/GameProjectGeneration/Public/GameProjectUtils.h
+++ b/Engine/Source/Editor/GameProjectGeneration/Public/GameProjectUtils.h
@@ -284,6 +284,9 @@ public:
 	/** Checks the name for illegal characters */
 	static bool NameContainsOnlyLegalCharacters(const FString& TestName, FString& OutIllegalCharacters);
 
+	/** Checks the name to see if it matches platform module names */
+	static bool NameMatchesPlatformModuleName(const FString& TestName);
+
 	/** Returns a list of #include lines formed from InList */
 	static FString MakeIncludeList(const TArray<FString>& InList);
 
