The Project file out of date notification now fades away after 10 seconds.

This now makes sense to change now that dialogs don't fade out during loading.

#rnx

diff --git a/Engine/Source/Editor/GameProjectGeneration/Private/GameProjectUtils.cpp b/Engine/Source/Editor/GameProjectGeneration/Private/GameProjectUtils.cpp
--- a/Engine/Source/Editor/GameProjectGeneration/Private/GameProjectUtils.cpp
+++ b/Engine/Source/Editor/GameProjectGeneration/Private/GameProjectUtils.cpp
@@ -951,7 +951,8 @@ void GameProjectUtils::CheckForOutOfDateGameProjectFile()
 			const FText UpdateProjectCancelText = LOCTEXT("UpdateProjectFileCancel", "Not Now");
 
 			FNotificationInfo Info(UpdateProjectText);
-			Info.bFireAndForget = false;
+			Info.ExpireDuration = 10;
+			Info.bFireAndForget = true;
 			Info.bUseLargeFont = false;
 			Info.bUseThrobber = false;
 			Info.bUseSuccessFailIcons = false;
