New projects should start with D3D12 set to only support SM6.

diff --git a/Engine/Source/Editor/GameProjectGeneration/Private/GameProjectUtils.cpp b/Engine/Source/Editor/GameProjectGeneration/Private/GameProjectUtils.cpp
--- a/Engine/Source/Editor/GameProjectGeneration/Private/GameProjectUtils.cpp
+++ b/Engine/Source/Editor/GameProjectGeneration/Private/GameProjectUtils.cpp
@@ -2064,6 +2064,20 @@ void GameProjectUtils::AddHardwareConfigValues(const FProjectInformation& InProj
 		TEXT("DefaultGraphicsRHI"),
 		TEXT("DefaultGraphicsRHI_DX12"),
 		false /* ShouldReplaceExistingValue */);
+
+	// Force clear D3D12TargetedShaderFormats since the BaseEngine list can change at any time.
+	ConfigValues.Emplace(TEXT("DefaultEngine.ini"),
+		TEXT("/Script/WindowsTargetPlatform.WindowsTargetSettings"),
+		TEXT("!D3D12TargetedShaderFormats"),
+		TEXT("ClearArray"),
+		false /* ShouldReplaceExistingValue */);
+
+	// New projects always have DX12 only supporting SM6 by default on Windows
+	ConfigValues.Emplace(TEXT("DefaultEngine.ini"),
+		TEXT("/Script/WindowsTargetPlatform.WindowsTargetSettings"),
+		TEXT("+D3D12TargetedShaderFormats"),
+		TEXT("PCD3D_SM6"),
+		false /* ShouldReplaceExistingValue */);
 }
 
 bool GameProjectUtils::GenerateConfigFiles(const FProjectInformation& InProjectInfo, TArray<FString>& OutCreatedFiles, FText& OutFailReason, FGuid& OutProjectID)
