Enable Extended Luminance Range for new projects by default
- Fix logic to also apply to projects created from template (doesn't overwrite template defined config if it has one).
- Removed deprecated option in Renderer Settings.
- Cleaned up branches around LuminanceMaxFromLensAttenuation(), since it's already applied inside the function.
- Templates are also updated to use ExtendedDefaultLuminanceRange.

diff --git a/Engine/Source/Editor/GameProjectGeneration/Private/GameProjectUtils.cpp b/Engine/Source/Editor/GameProjectGeneration/Private/GameProjectUtils.cpp
--- a/Engine/Source/Editor/GameProjectGeneration/Private/GameProjectUtils.cpp
+++ b/Engine/Source/Editor/GameProjectGeneration/Private/GameProjectUtils.cpp
@@ -155,6 +155,16 @@ namespace
 			true /* ShouldReplaceExistingValue */);
 	}
 
+	void AddPostProcessingConfigValues(const FProjectInformation& InProjectInfo, TArray<FTemplateConfigValue>& ConfigValues)
+	{
+		// Enable support for ExtendDefaultLuminanceRange by default for new projects
+		ConfigValues.Emplace(TEXT("DefaultEngine.ini"),
+			TEXT("/Script/Engine.RendererSettings"),
+			TEXT("r.DefaultFeature.AutoExposure.ExtendDefaultLuminanceRange"),
+			TEXT("True"),
+			false /* ShouldReplaceExistingValue */);
+	}
+
 	/** Get the configuration values for raytracing if enabled. */
 	void AddRaytracingConfigValues(const FProjectInformation& InProjectInfo, TArray<FTemplateConfigValue>& ConfigValues)
 	{
@@ -1795,6 +1805,7 @@ TOptional<FGuid> GameProjectUtils::CreateProjectFromTemplate(const FProjectInfor
 	AddLumenConfigValues(InProjectInfo, ConfigValuesToSet);
 	AddRaytracingConfigValues(InProjectInfo, ConfigValuesToSet);
 	AddNewProjectDefaultShadowConfigValues(InProjectInfo, ConfigValuesToSet);
+	AddPostProcessingConfigValues(InProjectInfo, ConfigValuesToSet);
 	AddWorldPartitionConfigValues(InProjectInfo, ConfigValuesToSet);
 	
 	TemplateDefs->AddConfigValues(ConfigValuesToSet, TemplateName, ProjectName, InProjectInfo.bShouldGenerateCode);
@@ -2096,13 +2107,6 @@ bool GameProjectUtils::GenerateConfigFiles(const FProjectInformation& InProjectI
 		FileContents += TEXT("[Audio]") LINE_TERMINATOR;
 		FileContents += TEXT("UseAudioMixer=True") LINE_TERMINATOR;
 
-		if (InProjectInfo.bForceExtendedLuminanceRange)
-		{
-			FileContents += LINE_TERMINATOR;
-			FileContents += TEXT("[/Script/Engine.RendererSettings]") LINE_TERMINATOR;
-			FileContents += TEXT("r.DefaultFeature.AutoExposure.ExtendDefaultLuminanceRange=True") LINE_TERMINATOR;
-		}
-
 		if (InProjectInfo.bCopyStarterContent)
 		{
 			FileContents += LINE_TERMINATOR;
@@ -2133,6 +2137,7 @@ bool GameProjectUtils::GenerateConfigFiles(const FProjectInformation& InProjectI
 		AddHardwareConfigValues(InProjectInfo, ConfigValuesToSet);
 		AddLumenConfigValues(InProjectInfo, ConfigValuesToSet);
 		AddNewProjectDefaultShadowConfigValues(InProjectInfo, ConfigValuesToSet);
+		AddPostProcessingConfigValues(InProjectInfo, ConfigValuesToSet);
 		AddRaytracingConfigValues(InProjectInfo, ConfigValuesToSet);
 		AddWorldPartitionConfigValues(InProjectInfo, ConfigValuesToSet);
 
diff --git a/Engine/Source/Editor/GameProjectGeneration/Private/SProjectDialog.cpp b/Engine/Source/Editor/GameProjectGeneration/Private/SProjectDialog.cpp
--- a/Engine/Source/Editor/GameProjectGeneration/Private/SProjectDialog.cpp
+++ b/Engine/Source/Editor/GameProjectGeneration/Private/SProjectDialog.cpp
@@ -1633,7 +1633,6 @@ FProjectInformation SProjectDialog::CreateProjectInfo() const
 	ProjectInfo.TemplateCategory = ActiveCategory;
 	ProjectInfo.bIsEnterpriseProject = SelectedTemplate->bIsEnterprise;
 	ProjectInfo.bIsBlankTemplate = SelectedTemplate->bIsBlankTemplate;
-	ProjectInfo.bForceExtendedLuminanceRange = SelectedTemplate->bIsBlankTemplate;
 
 	if (bShouldGenerateCode)
 	{
diff --git a/Engine/Source/Editor/GameProjectGeneration/Public/GameProjectUtils.h b/Engine/Source/Editor/GameProjectGeneration/Public/GameProjectUtils.h
--- a/Engine/Source/Editor/GameProjectGeneration/Public/GameProjectUtils.h
+++ b/Engine/Source/Editor/GameProjectGeneration/Public/GameProjectUtils.h
@@ -42,7 +42,6 @@ struct FProjectInformation
 	bool bCopyStarterContent = false;
 	bool bIsBlankTemplate = false;
 	bool bIsEnterpriseProject = false;
-	bool bForceExtendedLuminanceRange; // See "r.DefaultFeature.AutoExposure.ExtendDefaultLuminanceRange"
 
 	// These are all optional, because there is an additional state introduced by hiding the setting in the template.
 	// In this case, the template author has chosen not to give the user a choice,
diff --git a/Engine/Source/Runtime/Engine/Classes/Engine/RendererSettings.h b/Engine/Source/Runtime/Engine/Classes/Engine/RendererSettings.h
--- a/Engine/Source/Runtime/Engine/Classes/Engine/RendererSettings.h
+++ b/Engine/Source/Runtime/Engine/Classes/Engine/RendererSettings.h
@@ -657,13 +657,6 @@ class ENGINE_API URendererSettings : public UDeveloperSettings
 		ToolTip = "Default Value for auto exposure bias."))
 	float DefaultFeatureAutoExposureBias;
 
-	UE_DEPRECATED(4.26, "Extend Default Luminance Range is deprecated, and will be forced to ON at all times in future revisions.")
-	UPROPERTY(config, EditAnywhere, Category = DefaultSettings, meta = (
-		ConsoleVariable = "r.DefaultFeature.AutoExposure.ExtendDefaultLuminanceRange", DisplayName = "Extend default luminance range in Auto Exposure settings",
-		ToolTip = "Whether the default values for AutoExposure should support an extended range of scene luminance. Also changes the exposure settings to be expressed in EV100.",
-		ConfigRestartRequired=true))
-	uint32 bExtendDefaultLuminanceRangeInAutoExposureSettings: 1;
-
 	UPROPERTY(config, EditAnywhere, Category = DefaultSettings, meta = (
 		ConsoleVariable = "r.DefaultFeature.MotionBlur", DisplayName = "Motion Blur",
 		ToolTip = "Whether the default for MotionBlur is enabled or not (postprocess volume/camera/game setting can still override and enable or disable it independently)"))
diff --git a/Engine/Source/Runtime/Engine/Private/SceneView.cpp b/Engine/Source/Runtime/Engine/Private/SceneView.cpp
--- a/Engine/Source/Runtime/Engine/Private/SceneView.cpp
+++ b/Engine/Source/Runtime/Engine/Private/SceneView.cpp
@@ -187,8 +187,8 @@ static TAutoConsoleVariable<int32> CVarDefaultAutoExposureExtendDefaultLuminance
 	TEXT("Whether the default values for AutoExposure should support an extended range of scene luminance.\n")
 	TEXT("This also change the PostProcessSettings.Exposure.MinBrightness, MaxBrightness, HistogramLogMin and HisogramLogMax\n")
 	TEXT("to be expressed in EV100 values instead of in Luminance and Log2 Luminance.\n")
-	TEXT(" 0: Legacy range (default)\n")
-	TEXT(" 1: Extended range"));
+	TEXT(" 0: Legacy range (UE4 default)\n")
+	TEXT(" 1: Extended range (UE5 default)"));
 
 static TAutoConsoleVariable<int32> CVarDefaultMotionBlur(
 	TEXT("r.DefaultFeature.MotionBlur"),
diff --git a/Engine/Source/Runtime/Renderer/Private/PostProcess/PostProcessEyeAdaptation.cpp b/Engine/Source/Runtime/Renderer/Private/PostProcess/PostProcessEyeAdaptation.cpp
--- a/Engine/Source/Runtime/Renderer/Private/PostProcess/PostProcessEyeAdaptation.cpp
+++ b/Engine/Source/Runtime/Renderer/Private/PostProcess/PostProcessEyeAdaptation.cpp
@@ -247,16 +247,14 @@ bool IsAutoExposureDebugMode(const FViewInfo& View)
 
 float CalculateFixedAutoExposure(const FViewInfo& View)
 {
-	const bool bExtendedLuminanceRange = IsExtendLuminanceRangeEnabled();
-	const float LuminanceMax = bExtendedLuminanceRange ? LuminanceMaxFromLensAttenuation() : 1.0f;
+	const float LuminanceMax = LuminanceMaxFromLensAttenuation();
 	return EV100ToLuminance(LuminanceMax, View.Family->ExposureSettings.FixedEV100);
 }
 
 // on mobile, we are never using the Physical Camera, which is why we need the bForceDisablePhysicalCamera
 float CalculateManualAutoExposure(const FViewInfo& View, bool bForceDisablePhysicalCamera)
 {
-	const bool bExtendedLuminanceRange = IsExtendLuminanceRangeEnabled();
-	const float LuminanceMax = bExtendedLuminanceRange ? LuminanceMaxFromLensAttenuation() : 1.0f;
+	const float LuminanceMax = LuminanceMaxFromLensAttenuation();
 
 	const FPostProcessSettings& Settings = View.FinalPostProcessSettings;
 
@@ -277,7 +275,7 @@ FEyeAdaptationParameters GetEyeAdaptationParameters(const FViewInfo& View, ERHIF
 
 	const EAutoExposureMethod AutoExposureMethod = GetAutoExposureMethod(View);
 
-	const float LuminanceMax = bExtendedLuminanceRange ? LuminanceMaxFromLensAttenuation() : 1.0f;
+	const float LuminanceMax = LuminanceMaxFromLensAttenuation();
 
 	const float PercentToScale = 0.01f;
 
