Update message dialogs to use a clipboard button instead of a "Copy Message" text hyperlink.

diff --git a/Engine/Source/Developer/ToolWidgets/Private/Dialog/SMessageDialog.cpp b/Engine/Source/Developer/ToolWidgets/Private/Dialog/SMessageDialog.cpp
--- a/Engine/Source/Developer/ToolWidgets/Private/Dialog/SMessageDialog.cpp
+++ b/Engine/Source/Developer/ToolWidgets/Private/Dialog/SMessageDialog.cpp
@@ -7,7 +7,8 @@
 #include "Styling/AppStyle.h"
 #include "Styling/SlateBrush.h"
 #include "Widgets/Text/STextBlock.h"
-#include "Widgets/Input/SHyperlink.h"
+#include "Widgets/Input/SButton.h"
+#include "Widgets/Images/SImage.h"
 
 void SMessageDialog::Construct(const FArguments& InArgs)
 {
@@ -39,11 +40,17 @@ void SMessageDialog::Construct(const FArguments& InArgs)
 		.OnClosed(InArgs._OnClosed)
 		.BeforeButtons()
 		[
-			SNew(SHyperlink)
-				.OnNavigate(this, &SMessageDialog::CopyMessageToClipboard)
-				// Re-use existing localization of SChoiceDialog
-				.Text( NSLOCTEXT("SChoiceDialog", "CopyMessageHyperlink", "Copy Message") )
-				.ToolTipText( NSLOCTEXT("SChoiceDialog", "CopyMessageTooltip", "Copy the text in this message to the clipboard (CTRL+C)") )
+			SNew(SButton)
+			.ButtonStyle(FAppStyle::Get(), "SimpleButton")
+			.OnClicked(this, &SMessageDialog::OnCopyMessage)
+			.ToolTipText(NSLOCTEXT("SChoiceDialog", "CopyMessageTooltip", "Copy the text in this message to the clipboard (CTRL+C)"))
+			.ContentPadding(2.f)
+			.Content()
+			[
+				SNew(SImage)
+				.Image(FAppStyle::Get().GetBrush("Icons.Clipboard"))
+				.ColorAndOpacity(FSlateColor::UseForeground())
+			]
 		]
 		);
 }
@@ -60,6 +67,12 @@ FReply SMessageDialog::OnKeyDown(const FGeometry& MyGeometry, const FKeyEvent& I
 	return FReply::Unhandled();
 }
 
+FReply SMessageDialog::OnCopyMessage()
+{
+	CopyMessageToClipboard();
+	return FReply::Handled();
+}
+
 void SMessageDialog::CopyMessageToClipboard()
 {
 	FPlatformApplicationMisc::ClipboardCopy(*Message.ToString());
diff --git a/Engine/Source/Developer/ToolWidgets/Public/Dialog/SMessageDialog.h b/Engine/Source/Developer/ToolWidgets/Public/Dialog/SMessageDialog.h
--- a/Engine/Source/Developer/ToolWidgets/Public/Dialog/SMessageDialog.h
+++ b/Engine/Source/Developer/ToolWidgets/Public/Dialog/SMessageDialog.h
@@ -7,7 +7,7 @@
 
 /**
  * Special case of SCustomDialog dedicated to only displaying text messages.
- * This class enforces uniform style and also adds a hyperlink for copying the message.
+ * This class enforces uniform style and also adds a button for copying the message.
  */
 class TOOLWIDGETS_API SMessageDialog : public SCustomDialog
 {
@@ -61,6 +61,8 @@ public:
 
 	virtual	FReply OnKeyDown(const FGeometry& MyGeometry, const FKeyEvent& InKeyEvent) override;
 
+	virtual	FReply OnCopyMessage();
+
 private:
 
 	FText Message;
diff --git a/Engine/Source/Editor/EditorStyle/Private/StarshipStyle.cpp b/Engine/Source/Editor/EditorStyle/Private/StarshipStyle.cpp
--- a/Engine/Source/Editor/EditorStyle/Private/StarshipStyle.cpp
+++ b/Engine/Source/Editor/EditorStyle/Private/StarshipStyle.cpp
@@ -426,6 +426,7 @@ void FStarshipEditorStyle::FStyle::SetupGeneralStyles()
 		Set("Icons.Pinned", new IMAGE_BRUSH_SVG("Starship/Common/Pinned", Icon16x16));
 		Set("Icons.Unpinned", new IMAGE_BRUSH_SVG("Starship/Common/Unpinned", Icon16x16));
 		Set("Icons.Tools", new IMAGE_BRUSH_SVG("Starship/Common/EditorModes", Icon16x16));
+		Set("Icons.Clipboard", new IMAGE_BRUSH_SVG("Starship/Common/Clipboard", Icon16x16));
 
 
 		Set("Icons.Toolbar.Play", new IMAGE_BRUSH_SVG("Starship/Common/play", Icon20x20));
diff --git a/Engine/Source/Editor/GameProjectGeneration/Private/SVerbChoiceDialog.cpp b/Engine/Source/Editor/GameProjectGeneration/Private/SVerbChoiceDialog.cpp
--- a/Engine/Source/Editor/GameProjectGeneration/Private/SVerbChoiceDialog.cpp
+++ b/Engine/Source/Editor/GameProjectGeneration/Private/SVerbChoiceDialog.cpp
@@ -20,6 +20,7 @@
 #include "Types/SlateEnums.h"
 #include "Widgets/Input/SButton.h"
 #include "Widgets/Input/SHyperlink.h"
+#include "Widgets/Images/SImage.h"
 #include "Widgets/Layout/SBorder.h"
 #include "Widgets/Layout/SScrollBox.h"
 #include "Widgets/Layout/SUniformGridPanel.h"
@@ -72,7 +73,7 @@ void SVerbChoiceDialog::Construct( const FArguments& InArgs )
 	Buttons = InArgs._Buttons;
 	
 	TSharedPtr<SUniformGridPanel> ButtonBox;
-	TSharedPtr<SUniformGridPanel> HyperlinksBox;
+	TSharedPtr<SHorizontalBox> HyperlinksBox;
 
 	this->ChildSlot
 		[	
@@ -105,16 +106,31 @@ void SVerbChoiceDialog::Construct( const FArguments& InArgs )
 						[
 							SNew(SHorizontalBox)
 
+							+SHorizontalBox::Slot()
+								.AutoWidth()
+								.HAlign(HAlign_Left)
+								.VAlign(VAlign_Center)
+								.Padding(FMargin(12.0f, 0.f, 9.0f, 0.f))
+								[
+									SNew(SButton)
+									.ButtonStyle(FAppStyle::Get(), "SimpleButton")
+									.OnClicked(this, &SVerbChoiceDialog::HandleCopyMessageButtonClicked)
+									.ToolTipText(NSLOCTEXT("SVerbChoiceDialog", "CopyMessageTooltip", "Copy the text in this message to the clipboard (CTRL+C)"))
+									.ContentPadding(2.f)
+									.Content()
+									[
+										SNew(SImage)
+										.Image(FAppStyle::Get().GetBrush("Icons.Clipboard"))
+										.ColorAndOpacity(FSlateColor::UseForeground())
+									]
+								]
+
 							+ SHorizontalBox::Slot()
 								.FillWidth(1.0f)
 								.HAlign(HAlign_Left)
 								.VAlign(VAlign_Center)
-								.Padding(5.0f)
 								[
-									SAssignNew( HyperlinksBox, SUniformGridPanel )
-										.SlotPadding(FAppStyle::GetMargin("StandardDialog.SlotPadding"))
-										.MinDesiredSlotWidth(FAppStyle::GetFloat("StandardDialog.MinDesiredSlotWidth"))
-										.MinDesiredSlotHeight(FAppStyle::GetFloat("StandardDialog.MinDesiredSlotHeight"))
+									SAssignNew( HyperlinksBox, SHorizontalBox )
 								]
 
 							+ SHorizontalBox::Slot()
@@ -132,17 +148,12 @@ void SVerbChoiceDialog::Construct( const FArguments& InArgs )
 				]
 		];
 
-	HyperlinksBox->AddSlot(0, 0)
-		[
-			SNew(SHyperlink)
-				.OnNavigate(this, &SVerbChoiceDialog::HandleCopyMessageHyperlinkNavigate)
-				.Text( NSLOCTEXT("SVerbChoiceDialog", "CopyMessageHyperlink", "Copy Message") )
-				.ToolTipText( NSLOCTEXT("SVerbChoiceDialog", "CopyMessageTooltip", "Copy the text in this message to the clipboard (CTRL+C)") )
-		];
-
 	for(int32 Idx = 0; Idx < Hyperlinks.Get().Num(); Idx++)
 	{
-		HyperlinksBox->AddSlot(Idx + 1, 0)
+		HyperlinksBox->AddSlot()
+			.AutoWidth()
+			.HAlign(HAlign_Left)
+			.Padding(FAppStyle::GetMargin("StandardDialog.SlotPadding"))
 			[
 				SNew(SHyperlink)
 					.Text( Hyperlinks.Get()[Idx] )
@@ -193,9 +204,10 @@ void SVerbChoiceDialog::CopyMessageToClipboard( )
 	FPlatformApplicationMisc::ClipboardCopy( *Message.Get().ToString() );
 }
 
-void SVerbChoiceDialog::HandleCopyMessageHyperlinkNavigate( )
+FReply SVerbChoiceDialog::HandleCopyMessageButtonClicked( )
 {
 	CopyMessageToClipboard();
+	return FReply::Handled();
 }
 
 void SVerbChoiceDialog::HandleHyperlinkClicked( int32 InResponse )
diff --git a/Engine/Source/Editor/GameProjectGeneration/Private/SVerbChoiceDialog.h b/Engine/Source/Editor/GameProjectGeneration/Private/SVerbChoiceDialog.h
--- a/Engine/Source/Editor/GameProjectGeneration/Private/SVerbChoiceDialog.h
+++ b/Engine/Source/Editor/GameProjectGeneration/Private/SVerbChoiceDialog.h
@@ -46,8 +46,8 @@ protected:
 	void CopyMessageToClipboard( );
 
 private:
-	/** Handles clicking the 'Copy Message' hyperlink. */
-	void HandleCopyMessageHyperlinkNavigate( );
+	/** Handles clicking the 'Copy Message' button. */
+	FReply HandleCopyMessageButtonClicked( );
 
 	/** Handles clicking on a hyperlink. */
 	void HandleHyperlinkClicked( int32 InResponse );
diff --git a/Engine/Source/Editor/UnrealEd/Private/Dialogs/Dialogs.cpp b/Engine/Source/Editor/UnrealEd/Private/Dialogs/Dialogs.cpp
--- a/Engine/Source/Editor/UnrealEd/Private/Dialogs/Dialogs.cpp
+++ b/Engine/Source/Editor/UnrealEd/Private/Dialogs/Dialogs.cpp
@@ -21,7 +21,6 @@
 #include "Editor.h"
 #include "ObjectTools.h"
 #include "DesktopPlatformModule.h"
-#include "Widgets/Input/SHyperlink.h"
 #include "HAL/PlatformApplicationMisc.h"
 #include "HAL/PlatformMisc.h"
 
@@ -108,10 +107,17 @@ public:
 					.HAlign(HAlign_Left)
 					.VAlign(VAlign_Center)
 					[
-						SNew(SHyperlink)
-							.OnNavigate(this, &SChoiceDialog::HandleCopyMessageHyperlinkNavigate)
-							.Text( NSLOCTEXT("SChoiceDialog", "CopyMessageHyperlink", "Copy Message") )
-							.ToolTipText( NSLOCTEXT("SChoiceDialog", "CopyMessageTooltip", "Copy the text in this message to the clipboard (CTRL+C)") )
+						SNew(SButton)
+						.ButtonStyle(FAppStyle::Get(), "SimpleButton")
+						.OnClicked(this, &SChoiceDialog::HandleCopyMessageButtonClicked)
+						.ToolTipText(NSLOCTEXT("SChoiceDialog", "CopyMessageTooltip", "Copy the text in this message to the clipboard (CTRL+C)"))
+						.ContentPadding(2.f)
+						.Content()
+						[
+							SNew(SImage)
+							.Image(FAppStyle::Get().GetBrush("Icons.Clipboard"))
+							.ColorAndOpacity(FSlateColor::UseForeground())
+						]
 					]
 
 					+ SHorizontalBox::Slot()
@@ -285,10 +291,11 @@ private:
 		return FReply::Handled();
 	}
 
-	// Handles clicking the 'Copy Message' hyper link.
-	void HandleCopyMessageHyperlinkNavigate( )
+	// Handles clicking the 'Copy Message' button.
+	FReply HandleCopyMessageButtonClicked()
 	{
 		CopyMessageToClipboard();
+		return FReply::Handled();
 	}
 		
 public:
diff --git a/Engine/Source/Editor/UnrealEd/Private/Dialogs/SOutputLogDialog.cpp b/Engine/Source/Editor/UnrealEd/Private/Dialogs/SOutputLogDialog.cpp
--- a/Engine/Source/Editor/UnrealEd/Private/Dialogs/SOutputLogDialog.cpp
+++ b/Engine/Source/Editor/UnrealEd/Private/Dialogs/SOutputLogDialog.cpp
@@ -13,7 +13,7 @@
 #include "Widgets/Input/SButton.h"
 #include "Styling/AppStyle.h"
 #include "Editor.h"
-#include "Widgets/Input/SHyperlink.h"
+#include "Widgets/Images/SImage.h"
 #include "HAL/PlatformApplicationMisc.h"
 
 void SOutputLogDialog::Open( const FText& InTitle, const FText& InHeader, const FText& InLog, const FText& InFooter )
@@ -58,7 +58,8 @@ void SOutputLogDialog::Construct( const FArguments& InArgs )
 	Footer = InArgs._Footer.Get();
 	Buttons = InArgs._Buttons.Get();
 
-	MaxWidth = FSlateApplication::Get().GetPreferredWorkArea().GetSize().X * 0.8f;
+	const float DPIFactor = FPlatformApplicationMisc::GetDPIScaleFactorAtPoint(ParentWindow->GetPositionInScreen().X, ParentWindow->GetPositionInScreen().Y);
+	MaxWidth = FSlateApplication::Get().GetPreferredWorkArea().GetSize().X * 0.8f / DPIFactor;
 
 	TSharedPtr<SUniformGridPanel> ButtonBox;
 
@@ -121,10 +122,17 @@ void SOutputLogDialog::Construct( const FArguments& InArgs )
 								.VAlign(VAlign_Center)
 								.Padding(0.0f)
 								[
-									SNew(SHyperlink)
-										.OnNavigate(this, &SOutputLogDialog::HandleCopyMessageHyperlinkNavigate)
-										.Text( NSLOCTEXT("SOutputLogDialog", "CopyMessageHyperlink", "Copy Message") )
-										.ToolTipText( NSLOCTEXT("SOutputLogDialog", "CopyMessageTooltip", "Copy the text in this message to the clipboard (CTRL+C)") )
+									SNew(SButton)
+									.ButtonStyle(FAppStyle::Get(), "SimpleButton")
+									.OnClicked(this, &SOutputLogDialog::HandleCopyMessageButtonClicked)
+									.ToolTipText(NSLOCTEXT("SOutputLogDialog", "CopyMessageTooltip", "Copy the text in this message to the clipboard (CTRL+C)"))
+									.ContentPadding(2.f)
+									.Content()
+									[
+										SNew(SImage)
+										.Image(FAppStyle::Get().GetBrush("Icons.Clipboard"))
+										.ColorAndOpacity(FSlateColor::UseForeground())
+									]
 								]
 
 							+ SHorizontalBox::Slot()
@@ -194,7 +202,8 @@ FReply SOutputLogDialog::HandleButtonClicked( int32 InResponse )
 	return FReply::Handled();
 }
 
-void SOutputLogDialog::HandleCopyMessageHyperlinkNavigate( )
+FReply SOutputLogDialog::HandleCopyMessageButtonClicked( )
 {
 	CopyMessageToClipboard();
+	return FReply::Handled();
 }
diff --git a/Engine/Source/Editor/UnrealEd/Public/Dialogs/SOutputLogDialog.h b/Engine/Source/Editor/UnrealEd/Public/Dialogs/SOutputLogDialog.h
--- a/Engine/Source/Editor/UnrealEd/Public/Dialogs/SOutputLogDialog.h
+++ b/Engine/Source/Editor/UnrealEd/Public/Dialogs/SOutputLogDialog.h
@@ -46,8 +46,8 @@ private:
 	/** Handles clicking a message box button. */
 	FReply HandleButtonClicked( int32  );
 
-	/** Handles clicking the 'Copy Message' hyper link. */
-	void HandleCopyMessageHyperlinkNavigate( );
+	/** Handles clicking the 'Copy Message' button. */
+	FReply HandleCopyMessageButtonClicked( );
 
 	int32 Response;
 	TSharedPtr<SWindow> ParentWindow;
