When creating a project from a template, don't attempt to override the DirectX settings. For one, this code can't handle arrays correctly, and two, the templates should all be configured correctly.

#rnx

diff --git a/Engine/Source/Editor/GameProjectGeneration/Private/GameProjectUtils.cpp b/Engine/Source/Editor/GameProjectGeneration/Private/GameProjectUtils.cpp
--- a/Engine/Source/Editor/GameProjectGeneration/Private/GameProjectUtils.cpp
+++ b/Engine/Source/Editor/GameProjectGeneration/Private/GameProjectUtils.cpp
@@ -2069,26 +2069,30 @@ void GameProjectUtils::AddHardwareConfigValues(const FProjectInformation& InProj
 		}
 	}
 
-	// New projects always have DX12 by default on Windows
-	ConfigValues.Emplace(TEXT("DefaultEngine.ini"),
-		TEXT("/Script/WindowsTargetPlatform.WindowsTargetSettings"),
-		TEXT("DefaultGraphicsRHI"),
-		TEXT("DefaultGraphicsRHI_DX12"),
-		false /* ShouldReplaceExistingValue */);
-
-	// Force clear D3D12TargetedShaderFormats since the BaseEngine list can change at any time.
-	ConfigValues.Emplace(TEXT("DefaultEngine.ini"),
-		TEXT("/Script/WindowsTargetPlatform.WindowsTargetSettings"),
-		TEXT("!D3D12TargetedShaderFormats"),
-		TEXT("ClearArray"),
-		false /* ShouldReplaceExistingValue */);
-
-	// New projects always have DX12 only supporting SM6 by default on Windows
-	ConfigValues.Emplace(TEXT("DefaultEngine.ini"),
-		TEXT("/Script/WindowsTargetPlatform.WindowsTargetSettings"),
-		TEXT("+D3D12TargetedShaderFormats"),
-		TEXT("PCD3D_SM6"),
-		false /* ShouldReplaceExistingValue */);
+	// Don't override these settings for templates
+	if (InProjectInfo.TemplateFile.IsEmpty())
+	{
+		// New projects always have DX12 by default on Windows
+		ConfigValues.Emplace(TEXT("DefaultEngine.ini"),
+			TEXT("/Script/WindowsTargetPlatform.WindowsTargetSettings"),
+			TEXT("DefaultGraphicsRHI"),
+			TEXT("DefaultGraphicsRHI_DX12"),
+			false /* ShouldReplaceExistingValue */);
+
+		// Force clear D3D12TargetedShaderFormats since the BaseEngine list can change at any time.
+		ConfigValues.Emplace(TEXT("DefaultEngine.ini"),
+			TEXT("/Script/WindowsTargetPlatform.WindowsTargetSettings"),
+			TEXT("!D3D12TargetedShaderFormats"),
+			TEXT("ClearArray"),
+			false /* ShouldReplaceExistingValue */);
+
+		// New projects always have DX12 only supporting SM6 by default on Windows
+		ConfigValues.Emplace(TEXT("DefaultEngine.ini"),
+			TEXT("/Script/WindowsTargetPlatform.WindowsTargetSettings"),
+			TEXT("+D3D12TargetedShaderFormats"),
+			TEXT("PCD3D_SM6"),
+			false /* ShouldReplaceExistingValue */);
+	}
 }
 
 bool GameProjectUtils::GenerateConfigFiles(const FProjectInformation& InProjectInfo, TArray<FString>& OutCreatedFiles, FText& OutFailReason, FGuid& OutProjectID)
