The property bIsVariable of the Widget class is now set accordingly to a project setting (bAuthorizeAutomaticWidgetVariableCreation) in the UserInterface category, when creating a new one.
-it will always be false if bAuthorizeAutomaticWidgetVariableCreation is false.
-it will take the setting of the widget it bAuthorizeAutomaticWidgetVariableCreation is true.
The setting is set to true for existing projects (it keeps the current behavior), but will be set to false for new project.
It can be changed any time, and will only affects the new Widgets that will be created.

tests:
-drag n drop various new Widgets, changed value of the general setting on the fly: the IsVariable value respected the setting.
-PIE Lyra game, all blueprints compiles fine and various tested menus / HUD are ok.
-opened a project of current version, or older project converted in place: new setting is true.
-created a new project: setting is set to false.

diff --git a/Engine/Source/Editor/GameProjectGeneration/Private/GameProjectUtils.cpp b/Engine/Source/Editor/GameProjectGeneration/Private/GameProjectUtils.cpp
--- a/Engine/Source/Editor/GameProjectGeneration/Private/GameProjectUtils.cpp
+++ b/Engine/Source/Editor/GameProjectGeneration/Private/GameProjectUtils.cpp
@@ -226,6 +226,18 @@ namespace
 				true /* ShouldReplaceExistingValue */);
 		}
 	}
+
+	void AddUserInterfaceConfigValues(const FProjectInformation& InProjectInfo, TArray<FTemplateConfigValue>& ConfigValues)
+	{
+		if (InProjectInfo.bIsBlankTemplate)
+		{
+			ConfigValues.Emplace(TEXT("DefaultEngine.ini"),
+				TEXT("/Script/Engine.UserInterfaceSettings"),
+				TEXT("bAuthorizeAutomaticWidgetVariableCreation"),
+				TEXT("False"),
+				true /* ShouldReplaceExistingValue */);
+		}
+	}
 } // namespace <>
 
 FText FNewClassInfo::GetClassName() const
@@ -1807,6 +1819,7 @@ TOptional<FGuid> GameProjectUtils::CreateProjectFromTemplate(const FProjectInfor
 	AddNewProjectDefaultShadowConfigValues(InProjectInfo, ConfigValuesToSet);
 	AddPostProcessingConfigValues(InProjectInfo, ConfigValuesToSet);
 	AddWorldPartitionConfigValues(InProjectInfo, ConfigValuesToSet);
+	AddUserInterfaceConfigValues(InProjectInfo, ConfigValuesToSet);
 	
 	TemplateDefs->AddConfigValues(ConfigValuesToSet, TemplateName, ProjectName, InProjectInfo.bShouldGenerateCode);
 
@@ -2144,6 +2157,7 @@ bool GameProjectUtils::GenerateConfigFiles(const FProjectInformation& InProjectI
 		AddPostProcessingConfigValues(InProjectInfo, ConfigValuesToSet);
 		AddRaytracingConfigValues(InProjectInfo, ConfigValuesToSet);
 		AddWorldPartitionConfigValues(InProjectInfo, ConfigValuesToSet);
+		AddUserInterfaceConfigValues(InProjectInfo, ConfigValuesToSet);
 
 		if (!SaveConfigValues(InProjectInfo, ConfigValuesToSet, OutFailReason))
 		{
diff --git a/Engine/Source/Editor/UMGEditor/Private/Templates/WidgetTemplateClass.cpp b/Engine/Source/Editor/UMGEditor/Private/Templates/WidgetTemplateClass.cpp
--- a/Engine/Source/Editor/UMGEditor/Private/Templates/WidgetTemplateClass.cpp
+++ b/Engine/Source/Editor/UMGEditor/Private/Templates/WidgetTemplateClass.cpp
@@ -176,7 +176,7 @@ UWidget* FWidgetTemplateClass::CreateNamed(class UWidgetTree* Tree, FName NameOv
 	}
 
 	UWidget* NewWidget = Tree->ConstructWidget<UWidget>(WidgetClass.Get(), NameOverride);
-	NewWidget->OnCreationFromPalette();
+	NewWidget->CreatedFromPalette();
 
 	return NewWidget;
 }
diff --git a/Engine/Source/Runtime/Engine/Classes/Engine/UserInterfaceSettings.h b/Engine/Source/Runtime/Engine/Classes/Engine/UserInterfaceSettings.h
--- a/Engine/Source/Runtime/Engine/Classes/Engine/UserInterfaceSettings.h
+++ b/Engine/Source/Runtime/Engine/Classes/Engine/UserInterfaceSettings.h
@@ -193,6 +193,14 @@ public:
 	UPROPERTY(config, EditAnywhere, Category = "Widgets")
 	bool bLoadWidgetsOnDedicatedServer;
 
+	/** 
+	 * Setting to authorize or not automatic variable creation.
+	 * If true, variables will be created automatically, if the type created allows it. Drawback: it's easier to have a bad data architecture because various blueprint graph will have access to many variables.
+	 * If false, variables are never created automatically, and you have to create them manually on a case by case basis.
+	 */
+	UPROPERTY(config, EditAnywhere, Category = "Widgets")
+	bool bAuthorizeAutomaticWidgetVariableCreation;
+
 public:
 
 	virtual void PostInitProperties() override;
diff --git a/Engine/Source/Runtime/Engine/Private/UserInterfaceSettings.cpp b/Engine/Source/Runtime/Engine/Private/UserInterfaceSettings.cpp
--- a/Engine/Source/Runtime/Engine/Private/UserInterfaceSettings.cpp
+++ b/Engine/Source/Runtime/Engine/Private/UserInterfaceSettings.cpp
@@ -16,6 +16,7 @@ UUserInterfaceSettings::UUserInterfaceSettings(const FObjectInitializer& ObjectI
 	, RenderFocusRule(ERenderFocusRule::NavigationOnly)
 	, ApplicationScale(1)
 	, bLoadWidgetsOnDedicatedServer(true)
+	, bAuthorizeAutomaticWidgetVariableCreation(true)
 {
 	SectionName = TEXT("UI");
 }
diff --git a/Engine/Source/Runtime/UMG/Private/Components/Widget.cpp b/Engine/Source/Runtime/UMG/Private/Components/Widget.cpp
--- a/Engine/Source/Runtime/UMG/Private/Components/Widget.cpp
+++ b/Engine/Source/Runtime/UMG/Private/Components/Widget.cpp
@@ -14,6 +14,7 @@
 #include "UObject/UObjectHash.h"
 #include "UObject/UObjectIterator.h"
 #include "Engine/LocalPlayer.h"
+#include "Engine/UserInterfaceSettings.h"
 #include "Framework/Application/SlateApplication.h"
 #include "Widgets/Layout/SBorder.h"
 #include "Widgets/Layout/SSpacer.h"
@@ -1182,6 +1183,18 @@ const FText UWidget::GetPaletteCategory()
 	return LOCTEXT("Uncategorized", "Uncategorized");
 }
 
+void UWidget::CreatedFromPalette()
+{
+	// Allowing the variable creation if the setting allows it.
+	const UUserInterfaceSettings* UISettings = GetDefault<UUserInterfaceSettings>();
+	if (!UISettings->bAuthorizeAutomaticWidgetVariableCreation)
+	{
+		bIsVariable = false;
+	}
+
+	OnCreationFromPalette();
+}
+
 EVisibility UWidget::GetVisibilityInDesigner() const
 {
 	return bHiddenInDesigner ? EVisibility::Collapsed : EVisibility::Visible;
diff --git a/Engine/Source/Runtime/UMG/Public/Components/Widget.h b/Engine/Source/Runtime/UMG/Public/Components/Widget.h
--- a/Engine/Source/Runtime/UMG/Public/Components/Widget.h
+++ b/Engine/Source/Runtime/UMG/Public/Components/Widget.h
@@ -997,9 +997,12 @@ public:
 	/** Gets the palette category of the widget */
 	virtual const FText GetPaletteCategory();
 
+	/** Called by the palette after constructing a new widget. */
+	void CreatedFromPalette();
+
 	/**
-	 * Called by the palette after constructing a new widget, allows the widget to perform interesting 
-	 * default setup that we don't want to be UObject Defaults.
+	 * Called after constructing a new widget from the palette.
+	 * Allows the widget to perform interesting default setup that we don't want to be UObject Defaults.
 	 */
 	virtual void OnCreationFromPalette() { }
 
