Fix an issue with the project dialog temporaly showing the wrong options for a template.
By switching the notification mode of tiles, we avoid the weird state that appends when the selection changed but we didn't recive the selection change notification yet.

diff --git a/Engine/Source/Editor/GameProjectGeneration/Private/SProjectDialog.cpp b/Engine/Source/Editor/GameProjectGeneration/Private/SProjectDialog.cpp
--- a/Engine/Source/Editor/GameProjectGeneration/Private/SProjectDialog.cpp
+++ b/Engine/Source/Editor/GameProjectGeneration/Private/SProjectDialog.cpp
@@ -160,8 +160,11 @@ public:
 		check(InArgs._Item.IsValid())
 		Item = InArgs._Item;
 
+		STableRow::FArguments TableRowArguments;
+		TableRowArguments._SignalSelectionMode = ETableRowSignalSelectionMode::Instantaneous;
+
 		STableRow::Construct(
-			STableRow::FArguments()
+			TableRowArguments
 			.Style(FAppStyle::Get(), "ProjectBrowser.TableRow")
 			.Padding(2.0f)
 			.Content()
