Fix truncation warnings in GameProjectGeneration module

diff --git a/Engine/Source/Editor/GameProjectGeneration/Private/GameProjectUtils.cpp b/Engine/Source/Editor/GameProjectGeneration/Private/GameProjectUtils.cpp
--- a/Engine/Source/Editor/GameProjectGeneration/Private/GameProjectUtils.cpp
+++ b/Engine/Source/Editor/GameProjectGeneration/Private/GameProjectUtils.cpp
@@ -1670,7 +1670,7 @@ TOptional<FGuid> GameProjectUtils::CreateProjectFromTemplate(const FProjectInfor
 	SlowTask.EnterProgressFrame();
 	{
 		// Open a new feedback scope for the loop so we can report how far through the copy we are
-		FScopedSlowTask InnerSlowTask(FilesToCopy.Num());
+		FScopedSlowTask InnerSlowTask(static_cast<float>(FilesToCopy.Num()));
 		for ( const FString& SrcFilename : FilesToCopy )
 		{
 			// Update the progress
@@ -1759,7 +1759,7 @@ TOptional<FGuid> GameProjectUtils::CreateProjectFromTemplate(const FProjectInfor
 	SlowTask.EnterProgressFrame();
 	{
 		// Open a new feedback scope for the loop so we can report how far through the process we are
-		FScopedSlowTask InnerSlowTask(FilesThatNeedContentsReplaced.Num());
+		FScopedSlowTask InnerSlowTask(static_cast<float>(FilesThatNeedContentsReplaced.Num()));
 
 		// Open all files with the specified extensions and replace text
 		for ( const FString& FileToFix : FilesThatNeedContentsReplaced )
diff --git a/Engine/Source/Editor/GameProjectGeneration/Private/SNewClassDialog.cpp b/Engine/Source/Editor/GameProjectGeneration/Private/SNewClassDialog.cpp
--- a/Engine/Source/Editor/GameProjectGeneration/Private/SNewClassDialog.cpp
+++ b/Engine/Source/Editor/GameProjectGeneration/Private/SNewClassDialog.cpp
@@ -272,7 +272,7 @@ void SNewClassDialog::Construct( const FArguments& InArgs )
 	ChildSlot
 	[
 		SNew(SBorder)
-		.Padding(18)
+		.Padding(18.0f)
 		.BorderImage( FAppStyle::GetBrush("Docking.Tab.ContentAreaBrush") )
 		[
 			SNew(SVerticalBox)
@@ -298,7 +298,7 @@ void SNewClassDialog::Construct( const FArguments& InArgs )
 					SNew(SBorder)
 					.Visibility( this, &SNewClassDialog::GetGlobalErrorLabelVisibility )
 					.BorderImage(FAppStyle::Get().GetBrush("RoundedError"))
-					.Padding(FMargin(0, 5))
+					.Padding(FMargin(0.0f, 5.0f))
 					.Content()
 					[
 						SNew(SHorizontalBox)
@@ -339,7 +339,7 @@ void SNewClassDialog::Construct( const FArguments& InArgs )
 					// Title
 					+SVerticalBox::Slot()
 					.AutoHeight()
-					.Padding(0)
+					.Padding(0.0f)
 					[
 						SNew(STextBlock)
 						.Font(FAppStyle::Get().GetFontStyle("HeadingExtraSmall"))
@@ -363,7 +363,7 @@ void SNewClassDialog::Construct( const FArguments& InArgs )
 					// Page description and view options
 					+SVerticalBox::Slot()
 					.AutoHeight()
-					.Padding(0, 10)
+					.Padding(0.0f, 10.0f)
 					[
 						SNew(STextBlock)
 						.Text(
@@ -376,7 +376,7 @@ void SNewClassDialog::Construct( const FArguments& InArgs )
 					// Add Code list
 					+SVerticalBox::Slot()
 					.FillHeight(1.f)
-					.Padding(0, 10)
+					.Padding(0.0f, 10.0f)
 					[
 						SNew(SBorder)
 						.AddMetaData<FTutorialMetaData>(TEXT("AddCodeOptions"))
@@ -411,7 +411,7 @@ void SNewClassDialog::Construct( const FArguments& InArgs )
 
 					// Class selection
 					+SVerticalBox::Slot()
-					.Padding(30, 2)
+					.Padding(30.0f, 2.0f)
 					.AutoHeight()
 					[
 						SNew(SGridPanel)
@@ -484,7 +484,7 @@ void SNewClassDialog::Construct( const FArguments& InArgs )
 					// Title
 					+SVerticalBox::Slot()
 					.AutoHeight()
-					.Padding(0)
+					.Padding(0.0f)
 					[
 						SNew(STextBlock)
 						.Font(FAppStyle::Get().GetFontStyle("HeadingExtraSmall"))
@@ -494,13 +494,13 @@ void SNewClassDialog::Construct( const FArguments& InArgs )
 
 					+SVerticalBox::Slot()
 					.FillHeight(1.f)
-					.Padding(0, 10)
+					.Padding(0.0f, 10.0f)
 					[
 						SNew(SVerticalBox)
 
 						+SVerticalBox::Slot()
 						.AutoHeight()
-						.Padding(0, 0, 0, 5)
+						.Padding(0.0f, 0.0f, 0.0f, 5.0f)
 						[
 							SNew(STextBlock)
 							.Text(LOCTEXT("ClassNameDescription", "Enter a name for your new class. Class names may only contain alphanumeric characters, and may not contain a space.") )
@@ -508,7 +508,7 @@ void SNewClassDialog::Construct( const FArguments& InArgs )
 
 						+SVerticalBox::Slot()
 						.AutoHeight()
-						.Padding(0, 0, 0, 2)
+						.Padding(0.0f, 0.0f, 0.0f, 2.0f)
 						[
 							SNew(STextBlock)
 							.Text( ClassDomain == EClassDomain::Native ?
@@ -520,7 +520,7 @@ void SNewClassDialog::Construct( const FArguments& InArgs )
 						// Name Error label
 						+SVerticalBox::Slot()
 						.AutoHeight()
-						.Padding(0, 5)
+						.Padding(0.0f, 5.0f)
 						[
 							SNew(SWarningOrErrorBox)
 							.MessageStyle(EMessageStyle::Error)
@@ -541,7 +541,7 @@ void SNewClassDialog::Construct( const FArguments& InArgs )
 
 								+SVerticalBox::Slot()
 								.AutoHeight()
-								.Padding(0)
+								.Padding(0.0f)
 								[
 									SNew(SGridPanel)
 									.FillColumn(1, 1.0f)
@@ -549,7 +549,7 @@ void SNewClassDialog::Construct( const FArguments& InArgs )
 
 									+ SGridPanel::Slot(0, 0)
 									.VAlign(VAlign_Center)
-									.Padding(0, 0, 12, 0)
+									.Padding(0.0f, 0.0f, 12.0f, 0.0f)
 									[
 										SNew(STextBlock)
 										.Text(LOCTEXT("ClassTypeLabel", "Class Type"))
@@ -574,7 +574,7 @@ void SNewClassDialog::Construct( const FArguments& InArgs )
 									// Name label
 									+SGridPanel::Slot(0, 1)
 									.VAlign(VAlign_Center)
-									.Padding(0, 0, 12, 0)
+									.Padding(0.0f, 0.0f, 12.0f, 0.0f)
 									[
 										SNew(STextBlock)
 										.Text( LOCTEXT( "NameLabel", "Name" ) )
@@ -622,7 +622,7 @@ void SNewClassDialog::Construct( const FArguments& InArgs )
 									// Path label
 									+SGridPanel::Slot(0, 2)
 									.VAlign(ClassDomain == EClassDomain::Blueprint ? VAlign_Top : VAlign_Center)
-									.Padding(0, 0, 12, 0)
+									.Padding(0.0f, 0.0f, 12.0f, 0.0f)
 									[
 										SNew(STextBlock)
 										.Text( LOCTEXT( "PathLabel", "Path" ) )
@@ -637,11 +637,11 @@ void SNewClassDialog::Construct( const FArguments& InArgs )
 
 										// Blueprint Class asset path
 										+ SVerticalBox::Slot()
-										.Padding(0)
+										.Padding(0.0f)
 										[
 											SNew(SBox)
 											// Height override to force the visibility of a scrollbar (our parent is autoheight)
-											.HeightOverride(200)
+											.HeightOverride(200.0f)
 											.Visibility(ClassDomain == EClassDomain::Blueprint ? EVisibility::Visible : EVisibility::Collapsed)
 											[
 												ContentBrowser.CreatePathPicker(BlueprintPathConfig)
@@ -650,7 +650,7 @@ void SNewClassDialog::Construct( const FArguments& InArgs )
 
 										// Native C++ path
 										+ SVerticalBox::Slot()
-										.Padding(0)
+										.Padding(0.0f)
 										.AutoHeight()
 										[
 											SNew(SBox)
@@ -689,7 +689,7 @@ void SNewClassDialog::Construct( const FArguments& InArgs )
 									// Header output label
 									+SGridPanel::Slot(0, 3)
 									.VAlign(VAlign_Center)
-									.Padding(0, 0, 12, 0)
+									.Padding(0.0f, 0.0f, 12.0f, 0.0f)
 									[
 										SNew(STextBlock)
 										.Visibility(ClassDomain == EClassDomain::Blueprint ? EVisibility::Collapsed : EVisibility::Visible)
@@ -714,7 +714,7 @@ void SNewClassDialog::Construct( const FArguments& InArgs )
 									// Source output label
 									+SGridPanel::Slot(0, 4)
 									.VAlign(VAlign_Center)
-									.Padding(0, 0, 12, 0)
+									.Padding(0.0f, 0.0f, 12.0f, 0.0f)
 									[
 										SNew(STextBlock)
 										.Visibility(ClassDomain == EClassDomain::Blueprint ? EVisibility::Collapsed : EVisibility::Visible)
@@ -809,19 +809,19 @@ TSharedRef<ITableRow> SNewClassDialog::MakeParentClassListViewWidget(TSharedPtr<
 	const int32 ItemHeight = 64;
 	return
 		SNew( STableRow<TSharedPtr<FParentClassItem>>, OwnerTable )
-		.Padding(4)
+		.Padding(4.0f)
 		.Style(FAppStyle::Get(), "NewClassDialog.ParentClassListView.TableRow")
 		.ToolTip(IDocumentation::Get()->CreateToolTip(ClassFullDescription, nullptr, FEditorClassUtils::GetDocumentationPage(Class), FEditorClassUtils::GetDocumentationExcerpt(Class)))
 		[
 			SNew(SBox)
-			.HeightOverride(ItemHeight)
+			.HeightOverride(static_cast<float>(ItemHeight))
 			[
 				SNew(SHorizontalBox)
 				+ SHorizontalBox::Slot()
 				.AutoWidth()
 				.VAlign(VAlign_Center)
 				.HAlign(HAlign_Center)
-				.Padding(8)
+				.Padding(8.0f)
 				[
 					SNew(SBox)
 					.HeightOverride(ItemHeight / 2.0f)
@@ -835,7 +835,7 @@ TSharedRef<ITableRow> SNewClassDialog::MakeParentClassListViewWidget(TSharedPtr<
 				+ SHorizontalBox::Slot()
 				.FillWidth(1.0f)
 				.VAlign(VAlign_Center)
-				.Padding(4)
+				.Padding(4.0f)
 				[
 					SNew(SVerticalBox)
 					+SVerticalBox::Slot()
diff --git a/Engine/Source/Editor/GameProjectGeneration/Private/SProjectBrowser.cpp b/Engine/Source/Editor/GameProjectGeneration/Private/SProjectBrowser.cpp
--- a/Engine/Source/Editor/GameProjectGeneration/Private/SProjectBrowser.cpp
+++ b/Engine/Source/Editor/GameProjectGeneration/Private/SProjectBrowser.cpp
@@ -211,8 +211,8 @@ public:
 					.VAlign(VAlign_Center)
 					[
 						SNew(SBox)
-						.WidthOverride(102)
-						.HeightOverride(102)
+						.WidthOverride(102.0f)
+						.HeightOverride(102.0f)
 						[
 							SNew(SBorder)
 							.BorderImage(FAppStyle::Get().GetBrush("ProjectBrowser.ProjectTile.ThumbnailAreaBackground"))
@@ -554,11 +554,11 @@ TSharedRef<SToolTip> SProjectBrowser::MakeProjectToolTip( TSharedPtr<FProjectIte
 		.AutoWidth()
 		.HAlign(HAlign_Center)
 		.VAlign(VAlign_Center)
-		.Padding(FMargin(0, 0, 1, 0))
+		.Padding(FMargin(0.0f, 0.0f, 1.0f, 0.0f))
 		[
 			SNew(SBox)
-			.WidthOverride(20)
-			.HeightOverride(20)
+			.WidthOverride(20.0f)
+			.HeightOverride(20.0f)
 			[
 				SNew(SImage)
 				.Image(FAppStyle::GetBrush(PlatformInfo->GetIconStyleName(EPlatformIconSize::Normal)))
@@ -567,21 +567,21 @@ TSharedRef<SToolTip> SProjectBrowser::MakeProjectToolTip( TSharedPtr<FProjectIte
 	}
 
 	TSharedRef<SToolTip> Tooltip = SNew(SToolTip)
-	.TextMargin(1)
+	.TextMargin(1.0f)
 	.BorderImage( FAppStyle::GetBrush("ProjectBrowser.TileViewTooltip.ToolTipBorder") )
 	[
 		SNew(SBorder)
-		.Padding(6)
+		.Padding(6.0f)
 		.BorderImage( FAppStyle::GetBrush("ProjectBrowser.TileViewTooltip.NonContentBorder") )
 		[
 			SNew(SVerticalBox)
 
 			+SVerticalBox::Slot()
 			.AutoHeight()
-			.Padding(0, 0, 0, 4)
+			.Padding(0.0f, 0.0f, 0.0f, 4.0f)
 			[
 				SNew(SBorder)
-				.Padding(6)
+				.Padding(6.0f)
 				.BorderImage( FAppStyle::GetBrush("ProjectBrowser.TileViewTooltip.ContentBorder") )
 				[
 					SNew(SVerticalBox)
@@ -598,7 +598,7 @@ TSharedRef<SToolTip> SProjectBrowser::MakeProjectToolTip( TSharedPtr<FProjectIte
 					+SVerticalBox::Slot()
 					.AutoHeight()
 					.VAlign(VAlign_Center)
-					.Padding(0, 2, 0, 0)
+					.Padding(0.0f, 2.0f, 0.0f, 0.0f)
 					[
 						TargetPlatformIconsBox
 					]
@@ -609,7 +609,7 @@ TSharedRef<SToolTip> SProjectBrowser::MakeProjectToolTip( TSharedPtr<FProjectIte
 			.AutoHeight()
 			[
 				SNew(SBorder)
-				.Padding(6)
+				.Padding(6.0f)
 				.BorderImage( FAppStyle::GetBrush("ProjectBrowser.TileViewTooltip.ContentBorder") )
 				[
 					InfoBox
@@ -626,13 +626,13 @@ void SProjectBrowser::AddToToolTipInfoBox(const TSharedRef<SVerticalBox>& InfoBo
 {
 	InfoBox->AddSlot()
 	.AutoHeight()
-	.Padding(0, 1)
+	.Padding(0.0f, 1.0f)
 	[
 		SNew(SHorizontalBox)
 
 		+SHorizontalBox::Slot()
 		.AutoWidth()
-		.Padding(0, 0, 4, 0)
+		.Padding(0.0f, 0.0f, 4.0f, 0.0f)
 		[
 			SNew(STextBlock) .Text( FText::Format(LOCTEXT("ProjectBrowserTooltipFormat", "{0}:"), Key ) )
 			.ColorAndOpacity(FSlateColor::UseSubduedForeground())
diff --git a/Engine/Source/Editor/GameProjectGeneration/Private/SProjectDialog.cpp b/Engine/Source/Editor/GameProjectGeneration/Private/SProjectDialog.cpp
--- a/Engine/Source/Editor/GameProjectGeneration/Private/SProjectDialog.cpp
+++ b/Engine/Source/Editor/GameProjectGeneration/Private/SProjectDialog.cpp
@@ -184,11 +184,11 @@ public:
 						.VAlign(VAlign_Center)
 						[
 							SNew(SBox)
-							.WidthOverride(102)
-							.HeightOverride(102)
+							.WidthOverride(102.0f)
+							.HeightOverride(102.0f)
 							[
 								SNew(SBorder)
-								.Padding(0)
+								.Padding(0.0f)
 								.BorderImage(FAppStyle::Get().GetBrush("ProjectBrowser.ProjectTile.ThumbnailAreaBackground"))
 								.HAlign(HAlign_Center)
 								.VAlign(VAlign_Center)
@@ -397,7 +397,7 @@ public:
 				.OnClicked(InArgs._OnBrowseForFolder)
 				.ToolTipText(LOCTEXT("BrowseForFolder", "Browse for a folder"))
 				.Visibility(InArgs._IsReadOnly ? EVisibility::Hidden : EVisibility::Visible)
-				.ContentPadding(0)
+				.ContentPadding(0.0f)
 				[
 					SNew(SImage)
 					.Image(FAppStyle::Get().GetBrush("Icons.FolderClosed"))
@@ -666,7 +666,7 @@ TSharedRef<SWidget> SProjectDialog::MakeTemplateProjectView()
 		// Templates list
 		+ SVerticalBox::Slot()
 		.FillHeight(1.0f)
-		.Padding(0)
+		.Padding(0.0f)
 		[
 			SNew(SHorizontalBox)
 			+ SHorizontalBox::Slot()
@@ -687,7 +687,7 @@ TSharedRef<SWidget> SProjectDialog::MakeTemplateProjectView()
 			]
 			// Selected template details
 			+ SHorizontalBox::Slot()
-			.Padding(8, 0.0f)
+			.Padding(8.0f, 0.0f)
 			.AutoWidth()
 			[
 				SNew(SVerticalBox)
@@ -2117,7 +2117,7 @@ TSharedRef<ITableRow> SProjectDialog::ConstructMajorCategoryTableRow(TSharedPtr<
 		SNew(STableRow<TSharedPtr<FTemplateCategory>>, TableView)
 		.Style(FAppStyle::Get(), "ProjectBrowser.TableRow")
 		.ShowSelection(false)
-		.Padding(2);
+		.Padding(2.0f);
 	TWeakPtr<STableRow<TSharedPtr<FTemplateCategory>>> RowWeakPtr = Row;
 
 	Row->SetContent(
