Update BuildSettingsVersion default for new projects

diff --git a/Engine/Source/Editor/GameProjectGeneration/Private/GameProjectUtils.cpp b/Engine/Source/Editor/GameProjectGeneration/Private/GameProjectUtils.cpp
--- a/Engine/Source/Editor/GameProjectGeneration/Private/GameProjectUtils.cpp
+++ b/Engine/Source/Editor/GameProjectGeneration/Private/GameProjectUtils.cpp
@@ -2962,7 +2962,7 @@ bool GameProjectUtils::UpdateAdditionalPluginDirectory(const FString& InDir, con
 
 const TCHAR* GameProjectUtils::GetDefaultBuildSettingsVersion()
 {
-	return TEXT("BuildSettingsVersion.V2");
+	return TEXT("BuildSettingsVersion.V3");
 }
 
 bool GameProjectUtils::ReadTemplateFile(const FString& TemplateFileName, FString& OutFileContents, FText& OutFailReason)
