Remove leading "/" characters in Linux filepaths being used for sorting keys so that blank projects, which use the magic string "_1" for sorting, go back to being sorted first as intended.

diff --git a/Engine/Source/Editor/GameProjectGeneration/Private/SProjectDialog.cpp b/Engine/Source/Editor/GameProjectGeneration/Private/SProjectDialog.cpp
--- a/Engine/Source/Editor/GameProjectGeneration/Private/SProjectDialog.cpp
+++ b/Engine/Source/Editor/GameProjectGeneration/Private/SProjectDialog.cpp
@@ -1307,6 +1307,21 @@ TSharedPtr<FTemplateItem> SProjectDialog::GetSelectedTemplateItem() const
 	return nullptr;
 }
 
+namespace
+{
+	FString MakeSortKey(const FString& TemplateKey)
+	{
+		FString Output = TemplateKey;
+
+#if PLATFORM_LINUX
+		// Paths with a leading "/" would get sorted before the magic value used for blank projects: "_1"
+		Output.RemoveFromStart("/");
+#endif
+
+		return Output;
+	}
+}
+
 TMap<FName, TArray<TSharedPtr<FTemplateItem>> > SProjectDialog::FindTemplateProjects()
 {
 	// Clear the list out first - or we could end up with duplicates
@@ -1465,7 +1480,7 @@ TMap<FName, TArray<TSharedPtr<FTemplateItem>> > SProjectDialog::FindTemplateProj
 		Template->SortKey = TemplateDefs->SortKey;
 		if (Template->SortKey.IsEmpty())
 		{
-			Template->SortKey = TemplateKey;
+			Template->SortKey = MakeSortKey(TemplateKey);
 		}
 
 		FoundTemplates.Add(Template);
