Remove OculusVR + Steam VR references

diff --git a/Engine/Source/Editor/GameProjectGeneration/Private/GameProjectUtils.cpp b/Engine/Source/Editor/GameProjectGeneration/Private/GameProjectUtils.cpp
--- a/Engine/Source/Editor/GameProjectGeneration/Private/GameProjectUtils.cpp
+++ b/Engine/Source/Editor/GameProjectGeneration/Private/GameProjectUtils.cpp
@@ -95,30 +95,6 @@ constexpr const TCHAR GameProjectUtils::IncludePathFormatString[];
 
 namespace
 {
-	// @todo: This is currently not called from anywhere as this approach does not work for binary builds.
-	/** Set the state of XR plugins in OutProject based on the flags in InProjectInfo. */
-	void SetXRPluginStates(const FProjectInformation& InProjectInfo, FProjectDescriptor& OutProject)
-	{
-		static const FString XRPlugins[] = {
-			TEXT("OculusVR")};
-
-		if (InProjectInfo.bEnableXR.IsSet() && 
-			InProjectInfo.bEnableXR.GetValue() == false)
-		{
-			for (const FString& Plugin : XRPlugins)
-			{
-				int32 Index = OutProject.FindPluginReferenceIndex(Plugin);
-				if (Index == INDEX_NONE)
-				{
-					Index = OutProject.Plugins.AddDefaulted();
-					OutProject.Plugins[Index].Name = Plugin;
-				}
-
-				OutProject.Plugins[Index].bEnabled = false;
-			}
-		}
-	}
-
 	/** Get the configuration values for enabling Lumen by default. */
 	void AddLumenConfigValues(const FProjectInformation& InProjectInfo, TArray<FTemplateConfigValue>& ConfigValues)
 	{
diff --git a/Templates/TP_AEC_CollabBP/TP_AEC_CollabBP.uproject b/Templates/TP_AEC_CollabBP/TP_AEC_CollabBP.uproject
--- a/Templates/TP_AEC_CollabBP/TP_AEC_CollabBP.uproject
+++ b/Templates/TP_AEC_CollabBP/TP_AEC_CollabBP.uproject
@@ -63,22 +63,6 @@
 				"Android",
 				"HoloLens"
 			]
-		},
-		{
-			"Name": "OculusVR",
-			"Enabled": false,
-			"SupportedTargetPlatforms": [
-				"Win32",
-				"Win64"
-			]
-		},
-		{
-			"Name": "SteamVR",
-			"Enabled": false,
-			"SupportedTargetPlatforms": [
-				"Win32",
-				"Win64"
-			]
 		}
 	]
 }
\ No newline at end of file
diff --git a/Templates/TP_DMXBP/TP_DMXBP.uproject b/Templates/TP_DMXBP/TP_DMXBP.uproject
--- a/Templates/TP_DMXBP/TP_DMXBP.uproject
+++ b/Templates/TP_DMXBP/TP_DMXBP.uproject
@@ -23,22 +23,6 @@
 		{
 			"Name": "DMXProtocol",
 			"Enabled": true
-		},
-		{
-			"Name": "OculusVR",
-			"Enabled": false,
-			"SupportedTargetPlatforms": [
-				"Win32",
-				"Win64"
-			]
-		},
-		{
-			"Name": "SteamVR",
-			"Enabled": false,
-			"SupportedTargetPlatforms": [
-				"Win32",
-				"Win64"
-			]
 		}
 	]
 }
\ No newline at end of file
