Look for parent class templates too when generating class header & class cpp files

Cherrypick CL 25662803

diff --git a/Engine/Source/Editor/GameProjectGeneration/Private/GameProjectUtils.cpp b/Engine/Source/Editor/GameProjectGeneration/Private/GameProjectUtils.cpp
--- a/Engine/Source/Editor/GameProjectGeneration/Private/GameProjectUtils.cpp
+++ b/Engine/Source/Editor/GameProjectGeneration/Private/GameProjectUtils.cpp
@@ -3133,13 +3133,16 @@ bool GameProjectUtils::GenerateClassHeaderFile(const FString& NewHeaderFileName,
 	{
 		if (UClassTemplateEditorSubsystem* TemplateSubsystem = GEditor->GetEditorSubsystem<UClassTemplateEditorSubsystem>())
 		{
-			const UClass* BaseClass = ParentClassInfo.BaseClass;
-			if (const UClassTemplate* ClassTemplate = TemplateSubsystem->FindClassTemplate(ParentClassInfo.BaseClass))
+			for (const UClass* BaseClass = ParentClassInfo.BaseClass; BaseClass != nullptr; BaseClass = BaseClass->GetSuperClass())
 			{
-				bTemplateFound = ClassTemplate->ReadHeader(Template, OutFailReason);
-				if (!bTemplateFound)
+				if (const UClassTemplate* ClassTemplate = TemplateSubsystem->FindClassTemplate(BaseClass))
 				{
-					return false;
+					bTemplateFound = ClassTemplate->ReadHeader(Template, OutFailReason);
+					if (!bTemplateFound)
+					{
+						return false;
+					}
+					break;
 				}
 			}
 		}
@@ -3298,13 +3301,16 @@ bool GameProjectUtils::GenerateClassCPPFile(const FString& NewCPPFileName, const
 	{
 		if (UClassTemplateEditorSubsystem* TemplateSubsystem = GEditor->GetEditorSubsystem<UClassTemplateEditorSubsystem>())
 		{
-			const UClass* BaseClass = ParentClassInfo.BaseClass;
-			if (const UClassTemplate* ClassTemplate = TemplateSubsystem->FindClassTemplate(ParentClassInfo.BaseClass))
+			for (const UClass* BaseClass = ParentClassInfo.BaseClass; BaseClass != nullptr; BaseClass = BaseClass->GetSuperClass())
 			{
-				bTemplateFound = ClassTemplate->ReadSource(Template, OutFailReason);
-				if (!bTemplateFound)
+				if (const UClassTemplate* ClassTemplate = TemplateSubsystem->FindClassTemplate(BaseClass))
 				{
-					return false;
+					bTemplateFound = ClassTemplate->ReadSource(Template, OutFailReason);
+					if (!bTemplateFound)
+					{
+						return false;
+					}
+					break;
 				}
 			}
 		}
