Filter out program projects from the project browser

diff --git a/Engine/Source/Editor/GameProjectGeneration/Private/SProjectBrowser.cpp b/Engine/Source/Editor/GameProjectGeneration/Private/SProjectBrowser.cpp
--- a/Engine/Source/Editor/GameProjectGeneration/Private/SProjectBrowser.cpp
+++ b/Engine/Source/Editor/GameProjectGeneration/Private/SProjectBrowser.cpp
@@ -842,7 +842,7 @@ FReply SProjectBrowser::FindProjects()
 	const TArray<FString>& NativeProjectFiles = DefaultProjectDictionary.GetProjectPaths();
 	for(const FString& ProjectFile : NativeProjectFiles)
 	{
-		if (!ProjectFile.Contains(TEXT("/Templates/")))
+		if (!ProjectFile.Contains(TEXT("/Templates/")) && !ProjectFile.Contains(TEXT("/Programs/")))
 		{
 			AllFoundProjectFiles.Add(ProjectFile);
 		}
