Fix C++ project creation error: Failed to open selected source code accessor 'Xcode'

diff --git a/Engine/Source/Editor/GameProjectGeneration/Private/SProjectDialog.cpp b/Engine/Source/Editor/GameProjectGeneration/Private/SProjectDialog.cpp
--- a/Engine/Source/Editor/GameProjectGeneration/Private/SProjectDialog.cpp
+++ b/Engine/Source/Editor/GameProjectGeneration/Private/SProjectDialog.cpp
@@ -1820,8 +1820,29 @@ void SProjectDialog::DisplayError(const FText& ErrorText)
 bool SProjectDialog::OpenCodeIDE(const FString& ProjectFile)
 {
 	FText FailReason;
-
-	if (GameProjectUtils::OpenCodeIDE(ProjectFile, FailReason))
+    
+#if PLATFORM_MAC
+    // Modern Xcode projects are different based on Desktop/Mobile
+    FString Extension = FPaths::GetExtension(ProjectFile);
+    FString ModernXcodeProjectFile = ProjectFile;
+    if (SelectedHardwareClassTarget == EHardwareClass::Desktop)
+    {
+        ModernXcodeProjectFile.RemoveFromEnd(TEXT(".") + Extension);
+        ModernXcodeProjectFile += TEXT(" (Mac).") + Extension;
+    }
+    else if (SelectedHardwareClassTarget == EHardwareClass::Mobile)
+    {
+        ModernXcodeProjectFile.RemoveFromEnd(TEXT(".") + Extension);
+        ModernXcodeProjectFile += TEXT(" (IOS).") + Extension;
+    }
+#endif
+    
+	if (
+#if PLATFORM_MAC
+        GameProjectUtils::OpenCodeIDE(ModernXcodeProjectFile, FailReason) ||
+        // if modern failed, try again with legacy project name
+#endif
+        GameProjectUtils::OpenCodeIDE(ProjectFile, FailReason))
 	{
 		// Successfully opened code editing IDE, the editor is closing
 		// Close this window in case something prevents the editor from closing (save dialog, quit confirmation, etc)
