SNewClassDialog fixes

* Correctly set the "IsExternallyReferenceable" state after creating a BP class
* Fixed the initial path not being selected in the path picker when creating a BP class
* Hide the "Class Type" label when creating a BP class, as it has no value
* Show the current resolved "Path" when creating a BP class, as not everything in the path picker may update it

#rnx

diff --git a/Engine/Source/Editor/GameProjectGeneration/GameProjectGeneration.Build.cs b/Engine/Source/Editor/GameProjectGeneration/GameProjectGeneration.Build.cs
--- a/Engine/Source/Editor/GameProjectGeneration/GameProjectGeneration.Build.cs
+++ b/Engine/Source/Editor/GameProjectGeneration/GameProjectGeneration.Build.cs
@@ -25,6 +25,7 @@ public class GameProjectGeneration : ModuleRules
 				"Analytics",
 				"ApplicationCore",
 				"AppFramework",
+				"AssetTools",
 				"ClassViewer",
 				"Core",
 				"CoreUObject",
diff --git a/Engine/Source/Editor/GameProjectGeneration/Private/SNewClassDialog.cpp b/Engine/Source/Editor/GameProjectGeneration/Private/SNewClassDialog.cpp
--- a/Engine/Source/Editor/GameProjectGeneration/Private/SNewClassDialog.cpp
+++ b/Engine/Source/Editor/GameProjectGeneration/Private/SNewClassDialog.cpp
@@ -33,6 +33,7 @@
 #include "Widgets/Input/SHyperlink.h"
 #include "TutorialMetaData.h"
 #include "Kismet2/KismetEditorUtilities.h"
+#include "IAssetTools.h"
 #include "AssetRegistry/AssetRegistryModule.h"
 
 #include "Framework/Notifications/NotificationManager.h"
@@ -260,7 +261,7 @@ void SNewClassDialog::Construct( const FArguments& InArgs )
 	FPathPickerConfig BlueprintPathConfig;
 	if (ClassDomain == EClassDomain::Blueprint)
 	{
-		BlueprintPathConfig.DefaultPath = InArgs._InitialPath;
+		BlueprintPathConfig.DefaultPath = NewClassPath;
 		BlueprintPathConfig.bFocusSearchBoxWhenOpened = false;
 		BlueprintPathConfig.bAllowContextMenu = false;
 		BlueprintPathConfig.bAllowClassesFolder = false;
@@ -552,6 +553,7 @@ void SNewClassDialog::Construct( const FArguments& InArgs )
 									.Padding(0.0f, 0.0f, 12.0f, 0.0f)
 									[
 										SNew(STextBlock)
+										.Visibility(ClassDomain == EClassDomain::Blueprint ? EVisibility::Collapsed : EVisibility::Visible)
 										.Text(LOCTEXT("ClassTypeLabel", "Class Type"))
 									]
 
@@ -641,10 +643,22 @@ void SNewClassDialog::Construct( const FArguments& InArgs )
 										[
 											SNew(SBox)
 											// Height override to force the visibility of a scrollbar (our parent is autoheight)
-											.HeightOverride(200.0f)
+											.HeightOverride(220.0f)
 											.Visibility(ClassDomain == EClassDomain::Blueprint ? EVisibility::Visible : EVisibility::Collapsed)
 											[
-												ContentBrowser.CreatePathPicker(BlueprintPathConfig)
+												SNew(SVerticalBox)
+												
+												+SVerticalBox::Slot()
+												.AutoHeight()
+												[
+													SNew(STextBlock)
+													.Text(this, &SNewClassDialog::OnGetClassPathText)
+												]
+
+												+SVerticalBox::Slot()
+												[
+													ContentBrowser.CreatePathPicker(BlueprintPathConfig)
+												]
 											]
 										]
 
@@ -1131,6 +1145,9 @@ void SNewClassDialog::FinishClicked()
 				UBlueprint* NewBP = FKismetEditorUtilities::CreateBlueprint(const_cast<UClass*>(ParentClassInfo.BaseClass), Package, FName(*NewClassName), BPTYPE_Normal);
 				if (NewBP)
 				{
+					// Set the default "IsExternallyReferenceable" state
+					Package->SetIsExternallyReferenceable(IAssetTools::Get().GetCreateAssetsAsExternallyReferenceable());
+
 					// Notify the asset registry
 					FAssetRegistryModule::AssetCreated(NewBP);
 
@@ -1140,10 +1157,7 @@ void SNewClassDialog::FinishClicked()
 					OnAddedToProject.ExecuteIfBound( NewClassName, PackagePath, FString() );
 
 					// Sync the content browser to the new asset
-					TArray<UObject*> SyncAssets;
-					SyncAssets.Add(NewBP);
-					FContentBrowserModule& ContentBrowserModule = FModuleManager::LoadModuleChecked<FContentBrowserModule>("ContentBrowser");
-					ContentBrowserModule.Get().SyncBrowserToAssets(SyncAssets);
+					GEditor->SyncBrowserToObject(NewBP);
 
 					// Open the editor for the new asset
 					GEditor->GetEditorSubsystem<UAssetEditorSubsystem>()->OpenEditorForAsset(NewBP);
@@ -1242,10 +1256,7 @@ void SNewClassDialog::FinishClicked()
 				UClass* const NewClass = static_cast<UClass*>(FindObjectWithOuter(ClassPackage, UClass::StaticClass(), *NewClassName));
 				if ( NewClass )
 				{
-					TArray<UObject*> SyncAssets;
-					SyncAssets.Add(NewClass);
-					FContentBrowserModule& ContentBrowserModule = FModuleManager::LoadModuleChecked<FContentBrowserModule>("ContentBrowser");
-					ContentBrowserModule.Get().SyncBrowserToAssets(SyncAssets);
+					GEditor->SyncBrowserToObject(NewClass);
 				}
 			}
 
