Minor fix to TextBlock in New C++ Class dialog  to fix localization in Japanese.

diff --git a/Engine/Source/Editor/GameProjectGeneration/Private/SNewClassDialog.cpp b/Engine/Source/Editor/GameProjectGeneration/Private/SNewClassDialog.cpp
--- a/Engine/Source/Editor/GameProjectGeneration/Private/SNewClassDialog.cpp
+++ b/Engine/Source/Editor/GameProjectGeneration/Private/SNewClassDialog.cpp
@@ -318,6 +318,7 @@ void SNewClassDialog::Construct( const FArguments& InArgs )
 						[
 							SNew(STextBlock)
 							.Text( this, &SNewClassDialog::GetGlobalErrorLabelText )
+							.AutoWrapText(true)
 						]
 
 						+SHorizontalBox::Slot()
