SNewClassDialog: Pick a valid default path and prevent the user from picking read only paths from the path picker

diff --git a/Engine/Source/Editor/GameProjectGeneration/GameProjectGeneration.Build.cs b/Engine/Source/Editor/GameProjectGeneration/GameProjectGeneration.Build.cs
--- a/Engine/Source/Editor/GameProjectGeneration/GameProjectGeneration.Build.cs
+++ b/Engine/Source/Editor/GameProjectGeneration/GameProjectGeneration.Build.cs
@@ -47,7 +47,8 @@ public class GameProjectGeneration : ModuleRules
 				"LauncherPlatform",
 				"AddContentDialog",
 				"AudioMixer",
-				"AudioMixerCore"
+				"AudioMixerCore",
+				"ContentBrowserData"
 			}
 		);
 
diff --git a/Engine/Source/Editor/GameProjectGeneration/Private/SNewClassDialog.cpp b/Engine/Source/Editor/GameProjectGeneration/Private/SNewClassDialog.cpp
--- a/Engine/Source/Editor/GameProjectGeneration/Private/SNewClassDialog.cpp
+++ b/Engine/Source/Editor/GameProjectGeneration/Private/SNewClassDialog.cpp
@@ -123,6 +123,10 @@ void SNewClassDialog::Construct( const FArguments& InArgs )
 	if (ClassDomain == EClassDomain::Blueprint)
 	{
 		NewClassPath = InArgs._InitialPath.IsEmpty() ? TEXT("/Game") : InArgs._InitialPath;
+
+		// Pick a valid default path if the path is not writable
+		FContentBrowserModule& ContentBrowserModule = FModuleManager::LoadModuleChecked<FContentBrowserModule>("ContentBrowser");
+		NewClassPath = ContentBrowserModule.Get().GetInitialPathToSaveAsset(FContentBrowserItemPath(NewClassPath, EContentBrowserPathType::Internal)).GetInternalPathString();
 	}
 	else if(!InArgs._InitialPath.IsEmpty())
 	{
@@ -265,6 +269,7 @@ void SNewClassDialog::Construct( const FArguments& InArgs )
 		BlueprintPathConfig.bFocusSearchBoxWhenOpened = false;
 		BlueprintPathConfig.bAllowContextMenu = false;
 		BlueprintPathConfig.bAllowClassesFolder = false;
+		BlueprintPathConfig.bAllowReadOnlyFolders = false;
 		BlueprintPathConfig.OnPathSelected = FOnPathSelected::CreateSP(this, &SNewClassDialog::OnBlueprintPathSelected);
 	}
 
