SNewClassDialog: Fix incorrect conversion from class path + class name to FSoftObjectPath

diff --git a/Engine/Source/Editor/GameProjectGeneration/Private/SNewClassDialog.cpp b/Engine/Source/Editor/GameProjectGeneration/Private/SNewClassDialog.cpp
--- a/Engine/Source/Editor/GameProjectGeneration/Private/SNewClassDialog.cpp
+++ b/Engine/Source/Editor/GameProjectGeneration/Private/SNewClassDialog.cpp
@@ -1449,7 +1449,8 @@ void SNewClassDialog::UpdateInputValidity()
 		if (bLastInputValidityCheckSuccessful)
 		{
 			IAssetRegistry& AssetRegistry = FModuleManager::LoadModuleChecked<FAssetRegistryModule>(AssetRegistryConstants::ModuleName).Get();
-			if (AssetRegistry.GetAssetByObjectPath(FSoftObjectPath(NewClassPath / NewClassName)).IsValid())
+			const FSoftObjectPath ObjectPath(NewClassPath / NewClassName + "." + NewClassName);
+			if (AssetRegistry.GetAssetByObjectPath(ObjectPath).IsValid())
 			{
 				bLastInputValidityCheckSuccessful = false;
 				LastInputValidityErrorText = FText::Format(LOCTEXT("AssetAlreadyExists", "An asset called {0} already exists in {1}."), FText::FromString(NewClassName), FText::FromString(NewClassPath));
