[Scene Graph] Prevent creating Prefab blueprint with "Save As Prefab" dialog until valid path is selected once

- SPathPicker calls OnPathSelected in its Construct when path is valid(passes filters). This is controlled by a flag in FPathPickerConfig, off by default.

- SNewClassDialog disables Create button until OnPathSelected is called)

diff --git a/Engine/Source/Editor/ContentBrowser/Private/SPathPicker.cpp b/Engine/Source/Editor/ContentBrowser/Private/SPathPicker.cpp
--- a/Engine/Source/Editor/ContentBrowser/Private/SPathPicker.cpp
+++ b/Engine/Source/Editor/ContentBrowser/Private/SPathPicker.cpp
@@ -68,6 +68,11 @@ void SPathPicker::Construct( const FArguments& InArgs )
 		}
 
 		PathViewPtr->SetSelectedPaths({ VirtualPath.ToString() });
+
+		if (InArgs._PathPickerConfig.bNotifyDefaultPathSelected)
+		{
+			OnPathSelected.ExecuteIfBound(VirtualPath.ToString());
+		}
 	}
 }
 
diff --git a/Engine/Source/Editor/ContentBrowser/Public/IContentBrowserSingleton.h b/Engine/Source/Editor/ContentBrowser/Public/IContentBrowserSingleton.h
--- a/Engine/Source/Editor/ContentBrowser/Public/IContentBrowserSingleton.h
+++ b/Engine/Source/Editor/ContentBrowser/Public/IContentBrowserSingleton.h
@@ -382,6 +382,9 @@ struct FPathPickerConfig
 	/** Whether or not to show the favorites selector. */
 	bool bShowFavorites : 1;
 
+	/** Whether to call OnPathSelected during construction for DefaultPath if DefaultPath is allowed */
+	bool bNotifyDefaultPathSelected : 1;
+
 	FPathPickerConfig()
 		: bFocusSearchBoxWhenOpened(true)
 		, bAllowContextMenu(true)
@@ -390,6 +393,7 @@ struct FPathPickerConfig
 		, bAddDefaultPath(false)
 		, bOnPathSelectedPassesVirtualPaths(false)
 		, bShowFavorites(true)
+		, bNotifyDefaultPathSelected(false)
 	{}
 };
 
diff --git a/Engine/Source/Editor/GameProjectGeneration/Private/SNewClassDialog.cpp b/Engine/Source/Editor/GameProjectGeneration/Private/SNewClassDialog.cpp
--- a/Engine/Source/Editor/GameProjectGeneration/Private/SNewClassDialog.cpp
+++ b/Engine/Source/Editor/GameProjectGeneration/Private/SNewClassDialog.cpp
@@ -271,6 +271,7 @@ void SNewClassDialog::Construct( const FArguments& InArgs )
 		BlueprintPathConfig.bAllowClassesFolder = false;
 		BlueprintPathConfig.bAllowReadOnlyFolders = false;
 		BlueprintPathConfig.OnPathSelected = FOnPathSelected::CreateSP(this, &SNewClassDialog::OnBlueprintPathSelected);
+		BlueprintPathConfig.bNotifyDefaultPathSelected = true;
 	}
 
 	OnAddedToProject = InArgs._OnAddedToProject;
@@ -1100,6 +1101,7 @@ void SNewClassDialog::OnClassPathTextChanged(const FText& NewText)
 
 void SNewClassDialog::OnBlueprintPathSelected(const FString& NewPath)
 {
+	IsBlueprintPathSelected = true;
 	NewClassPath = NewPath;
 	UpdateInputValidity();
 }
@@ -1121,7 +1123,7 @@ void SNewClassDialog::CancelClicked()
 
 bool SNewClassDialog::CanFinish() const
 {
-	return bLastInputValidityCheckSuccessful && ParentClassInfo.IsSet() && (ClassDomain == EClassDomain::Blueprint || FSourceCodeNavigation::IsCompilerAvailable());
+	return bLastInputValidityCheckSuccessful && ParentClassInfo.IsSet() && (ClassDomain == EClassDomain::Blueprint || FSourceCodeNavigation::IsCompilerAvailable()) && (ClassDomain != EClassDomain::Blueprint || IsBlueprintPathSelected);
 }
 
 void SNewClassDialog::FinishClicked()
diff --git a/Engine/Source/Editor/GameProjectGeneration/Private/SNewClassDialog.h b/Engine/Source/Editor/GameProjectGeneration/Private/SNewClassDialog.h
--- a/Engine/Source/Editor/GameProjectGeneration/Private/SNewClassDialog.h
+++ b/Engine/Source/Editor/GameProjectGeneration/Private/SNewClassDialog.h
@@ -298,4 +298,7 @@ private:
 
 	/** Event called when code is succesfully added to the project */
 	FOnAddedToProject OnAddedToProject;
+
+	/** Flag set when PathPicker selects a path*/
+	bool IsBlueprintPathSelected = false;
 };
