[Misc Editor] New class dialog can use the given class name and parent class info to populate more descriptive strings for blueprint based types

diff --git a/Engine/Source/Editor/GameProjectGeneration/Private/SNewClassDialog.cpp b/Engine/Source/Editor/GameProjectGeneration/Private/SNewClassDialog.cpp
--- a/Engine/Source/Editor/GameProjectGeneration/Private/SNewClassDialog.cpp
+++ b/Engine/Source/Editor/GameProjectGeneration/Private/SNewClassDialog.cpp
@@ -290,11 +290,11 @@ void SNewClassDialog::Construct( const FArguments& InArgs )
 				SAssignNew( MainWizard, SWizard)
 				.ShowPageList(false)
 				.CanFinish(this, &SNewClassDialog::CanFinish)
-				.FinishButtonText( ClassDomain == EClassDomain::Native ? LOCTEXT("FinishButtonText_Native", "Create Class") : LOCTEXT("FinishButtonText_Blueprint", "Create Blueprint Class") )
+				.FinishButtonText( ClassDomain == EClassDomain::Native ? LOCTEXT("FinishButtonText_Native", "Create Class") : FText::Format(LOCTEXT("FinishButtonText_Blueprint", "Create {0} Class"), ParentClassInfo.IsSet() ? ParentClassInfo.GetClassName() : FText::FromStringView(TEXT("Blueprint"))))
 				.FinishButtonToolTip (
 					ClassDomain == EClassDomain::Native ?
 					LOCTEXT("FinishButtonToolTip_Native", "Creates the code files to add your new class.") : 
-					LOCTEXT("FinishButtonToolTip_Blueprint", "Creates the new Blueprint class based on the specified parent class.")
+					FText::Format(LOCTEXT("FinishButtonToolTip_Blueprint", "Creates the new class based on the specified parent {0} class."), ParentClassInfo.IsSet() ? ParentClassInfo.GetClassName() : FText::FromStringView(TEXT("Blueprint")))
 					)
 				.OnCanceled(this, &SNewClassDialog::CancelClicked)
 				.OnFinished(this, &SNewClassDialog::FinishClicked)
@@ -377,7 +377,7 @@ void SNewClassDialog::Construct( const FArguments& InArgs )
 						.Text(
 							ClassDomain == EClassDomain::Native ?
 							LOCTEXT("ChooseParentClassDescription_Native", "This will add a C++ header and source code file to your game project.") :
-							LOCTEXT("ChooseParentClassDescription_Blueprint", "This will add a new Blueprint class to your game project.")
+							FText::Format(LOCTEXT("ChooseParentClassDescription_Blueprint", "This will add a new class inheriting from {0} to your game project."), ParentClassInfo.IsSet() ? ParentClassInfo.GetClassName() : FText::FromStringView(TEXT("Blueprint")))
 						)
 					]
 
@@ -521,7 +521,7 @@ void SNewClassDialog::Construct( const FArguments& InArgs )
 							SNew(STextBlock)
 							.Text( ClassDomain == EClassDomain::Native ?
 								LOCTEXT("ClassNameDetails_Native", "When you click the \"Create\" button below, a header (.h) file and a source (.cpp) file will be made using this name.") :
-								LOCTEXT("ClassNameDetails_Blueprint", "When you click the \"Create\" button below, a new Blueprint class will be created.")
+								FText::Format(LOCTEXT("ClassNameDetails_Blueprint", "When you click the \"Create\" button below, a new class inheriting from {0} will be created."), ParentClassInfo.IsSet() ? ParentClassInfo.GetClassName() : FText::FromStringView(TEXT("Blueprint")))
 								)
 						]
 
@@ -1137,12 +1137,12 @@ void SNewClassDialog::FinishClicked()
 		if (!ParentClassInfo.BaseClass)
 		{
 			// @todo show fail reason in error label
-			FMessageDialog::Open(EAppMsgType::Ok, LOCTEXT("AddCodeFailed_Blueprint_NoBase", "No parent class has been specified. Failed to generate new Blueprint class."));
+			FMessageDialog::Open(EAppMsgType::Ok, FText::Format(LOCTEXT("AddCodeFailed_Blueprint_NoBase", "No parent class has been specified. Failed to generate new {0} class."), FText::FromString(NewClassName)));
 		}
 		else if (FindObject<UBlueprint>(nullptr, *PackagePath))
 		{
 			// @todo show fail reason in error label
-			FMessageDialog::Open(EAppMsgType::Ok, LOCTEXT("AddCodeFailed_Blueprint_AlreadyExists", "The chosen Blueprint class already exists, please try again with a different name."));
+			FMessageDialog::Open(EAppMsgType::Ok, FText::Format(LOCTEXT("AddCodeFailed_Blueprint_AlreadyExists", "The chosen class name ({0}) already exists, please try again with a different name."), FText::FromString(NewClassName)));
 		}
 		else if (!NewClassPath.IsEmpty() && !NewClassName.IsEmpty())
 		{
