Fixed issue with Deleted StarterContent getting copied over and over
Root cause: the update of bAddPacks to False in the Game config was not saved to disk
Solution: Added a "SectionsToSave" section to the DefaultGame.ini at creation to add the 'StartupActions' as persistable

#rnx

diff --git a/Engine/Source/Editor/GameProjectGeneration/Private/GameProjectUtils.cpp b/Engine/Source/Editor/GameProjectGeneration/Private/GameProjectUtils.cpp
--- a/Engine/Source/Editor/GameProjectGeneration/Private/GameProjectUtils.cpp
+++ b/Engine/Source/Editor/GameProjectGeneration/Private/GameProjectUtils.cpp
@@ -4257,6 +4257,15 @@ bool GameProjectUtils::InsertFeaturePacksIntoINIFile(const FProjectInformation&
 			FileOutput += PackList[iLine] + LINE_TERMINATOR;
 		}
 
+		// Register 'StartupActions' as a section to save if one of the default value of its entries
+		// is modified. Otherwise, the mechanism which set 'bAddPacks' to 'False' after the
+		// first load will not be persisted and the selected pack files will be loaded on each launch of the editor.
+		FileOutput += LINE_TERMINATOR;
+		FileOutput += TEXT("[SectionsToSave]");
+		FileOutput += LINE_TERMINATOR;
+		FileOutput += TEXT("+Section=StartupActions");
+		FileOutput += LINE_TERMINATOR;
+
 		if (!FFileHelper::SaveStringToFile(FileOutput, *IniFilename))
 		{
 			OutFailReason = LOCTEXT("FailedToWriteIni", "Could not write INI file to insert feature packs");
