Enable Ray Tracing by default on new projects.
- also remove "Ray Tracing" check box from new project dialog.

diff --git a/Engine/Source/Editor/GameProjectGeneration/Classes/TemplateProjectDefs.h b/Engine/Source/Editor/GameProjectGeneration/Classes/TemplateProjectDefs.h
--- a/Engine/Source/Editor/GameProjectGeneration/Classes/TemplateProjectDefs.h
+++ b/Engine/Source/Editor/GameProjectGeneration/Classes/TemplateProjectDefs.h
@@ -82,7 +82,7 @@ enum class ETemplateSetting
 	GraphicsPreset,
 	StarterContent,
 	XR,
-	Raytracing,
+	Raytracing_DEPRECATED,
 	All
 };
 
diff --git a/Engine/Source/Editor/GameProjectGeneration/Private/GameProjectUtils.cpp b/Engine/Source/Editor/GameProjectGeneration/Private/GameProjectUtils.cpp
--- a/Engine/Source/Editor/GameProjectGeneration/Private/GameProjectUtils.cpp
+++ b/Engine/Source/Editor/GameProjectGeneration/Private/GameProjectUtils.cpp
@@ -157,27 +157,23 @@ namespace
 	/** Get the configuration values for raytracing if enabled. */
 	void AddRaytracingConfigValues(const FProjectInformation& InProjectInfo, TArray<FTemplateConfigValue>& ConfigValues)
 	{
-		if (InProjectInfo.bEnableRaytracing.IsSet() && 
-			InProjectInfo.bEnableRaytracing.GetValue() == true)
-		{
-			ConfigValues.Emplace(TEXT("DefaultEngine.ini"),
-				TEXT("/Script/WindowsTargetPlatform.WindowsTargetSettings"),
-				TEXT("DefaultGraphicsRHI"),
-				TEXT("DefaultGraphicsRHI_DX12"),
-				true /* ShouldReplaceExistingValue */);
+		ConfigValues.Emplace(TEXT("DefaultEngine.ini"),
+			TEXT("/Script/WindowsTargetPlatform.WindowsTargetSettings"),
+			TEXT("DefaultGraphicsRHI"),
+			TEXT("DefaultGraphicsRHI_DX12"),
+			false /* ShouldReplaceExistingValue */);
 
-			ConfigValues.Emplace(TEXT("DefaultEngine.ini"),
-				TEXT("/Script/Engine.RendererSettings"),
-				TEXT("r.SkinCache.CompileShaders"),
-				TEXT("True"),
-				true /* ShouldReplaceExistingValue */);
+		ConfigValues.Emplace(TEXT("DefaultEngine.ini"),
+			TEXT("/Script/Engine.RendererSettings"),
+			TEXT("r.SkinCache.CompileShaders"),
+			TEXT("True"),
+			false /* ShouldReplaceExistingValue */);
 
-			ConfigValues.Emplace(TEXT("DefaultEngine.ini"),
-				TEXT("/Script/Engine.RendererSettings"),
-				TEXT("r.RayTracing"),
-				TEXT("True"),
-				true /* ShouldReplaceExistingValue */);
-		}
+		ConfigValues.Emplace(TEXT("DefaultEngine.ini"),
+			TEXT("/Script/Engine.RendererSettings"),
+			TEXT("r.RayTracing"),
+			TEXT("True"),
+			false /* ShouldReplaceExistingValue */);
 	}
 
 	void AddDefaultMapConfigValues(const FProjectInformation& InProjectInfo, TArray<FTemplateConfigValue>& ConfigValues)
diff --git a/Engine/Source/Editor/GameProjectGeneration/Private/SProjectDialog.cpp b/Engine/Source/Editor/GameProjectGeneration/Private/SProjectDialog.cpp
--- a/Engine/Source/Editor/GameProjectGeneration/Private/SProjectDialog.cpp
+++ b/Engine/Source/Editor/GameProjectGeneration/Private/SProjectDialog.cpp
@@ -1095,33 +1095,7 @@ TSharedRef<SWidget> SProjectDialog::MakeProjectOptionsWidget()
 			.AutoWrapText(true)
 			.DecoratorStyleSet(&FAppStyle::Get());
 	}
-#endif 
-
-	if (!HiddenSettings.Contains(ETemplateSetting::Raytracing))
-	{
-		ProjectOptionsBox->AddSlot()
-			.Padding(0.0f, 8.0f, 0.0f, 0.0f)
-			.AutoHeight()
-			[
-				SNew(SHorizontalBox)
-				.ToolTipText(LOCTEXT("ProjectDialog_RaytracingDescription", "Choose if real-time raytracing should be supported in the new project."))
-				+ SHorizontalBox::Slot()
-				.Padding(0.0f, 0.0f, 8.0f, 0.0f)
-				.VAlign(VAlign_Center)
-				.HAlign(HAlign_Right)
-				[
-					SNew(STextBlock)
-					.Text(LOCTEXT("RayTracing", "Raytracing"))
-				]
-				+ SHorizontalBox::Slot()
-				.VAlign(VAlign_Center)
-				[
-					SNew(SCheckBox)
-					.IsChecked(this, &SProjectDialog::OnGetRaytracingEnabledCheckState)
-					.OnCheckStateChanged(this, &SProjectDialog::OnSetRaytracingEnabled)
-				]
-			];
-	}
+#endif
 
 	return ProjectOptionsWidget;
 }
@@ -1254,11 +1228,6 @@ FReply SProjectDialog::OnCancel() const
 	return FReply::Handled();
 }
 
-void SProjectDialog::OnSetRaytracingEnabled(ECheckBoxState NewState)
-{
-	bEnableRaytracing = NewState == ECheckBoxState::Checked;
-}
-
 void SProjectDialog::OnSetBlueprintOrCppIndex(int32 Index)
 {
 	bShouldGenerateCode = Index == 1;
@@ -1685,11 +1654,6 @@ FProjectInformation SProjectDialog::CreateProjectInfo() const
 		{
 			ProjectInfo.bEnableXR = bEnableXR;
 		}
-
-		if (!HiddenSettings.Contains(ETemplateSetting::Raytracing))
-		{
-			ProjectInfo.bEnableRaytracing = bEnableRaytracing;
-		}
 	}
 
 	return MoveTemp(ProjectInfo);
diff --git a/Engine/Source/Editor/GameProjectGeneration/Private/SProjectDialog.h b/Engine/Source/Editor/GameProjectGeneration/Private/SProjectDialog.h
--- a/Engine/Source/Editor/GameProjectGeneration/Private/SProjectDialog.h
+++ b/Engine/Source/Editor/GameProjectGeneration/Private/SProjectDialog.h
@@ -168,13 +168,9 @@ private:
 	bool CanCreateProject() const;
 
 	FReply OnCancel() const;
-
-	ECheckBoxState OnGetRaytracingEnabledCheckState() const { return bEnableRaytracing ? ECheckBoxState::Checked : ECheckBoxState::Unchecked; }
 	
 	bool IsStarterContentAvailable() const {return bIsStarterContentAvailable; }
 
-	void OnSetRaytracingEnabled(ECheckBoxState NewState);
-
 	int32 OnGetBlueprintOrCppIndex() const { return bShouldGenerateCode ? 1 : 0; }
 	void OnSetBlueprintOrCppIndex(int32 Index);
 
@@ -257,9 +253,6 @@ private:
 	/** Whether or not to enable XR in the created project. */
 	bool bEnableXR;
 
-	/** Whether or not to enable Raytracing in the created project. */
-	bool bEnableRaytracing;
-
 	/** Whether or not we should use the blueprint or C++ version of this template. */
 	bool bShouldGenerateCode;
 
diff --git a/Engine/Source/Editor/GameProjectGeneration/Public/GameProjectUtils.h b/Engine/Source/Editor/GameProjectGeneration/Public/GameProjectUtils.h
--- a/Engine/Source/Editor/GameProjectGeneration/Public/GameProjectUtils.h
+++ b/Engine/Source/Editor/GameProjectGeneration/Public/GameProjectUtils.h
@@ -47,7 +47,6 @@ struct FProjectInformation
 	// In this case, the template author has chosen not to give the user a choice,
 	// so we must assume that the template already controls these settings explicitly.
 	TOptional<bool> bEnableXR;
-	TOptional<bool> bEnableRaytracing;
 
 	TOptional<EHardwareClass> TargetedHardware;
 	TOptional<EGraphicsPreset> DefaultGraphicsPerformance;
