Add New Class error message padding fix

diff --git a/Engine/Source/Editor/GameProjectGeneration/Private/SNewClassDialog.cpp b/Engine/Source/Editor/GameProjectGeneration/Private/SNewClassDialog.cpp
--- a/Engine/Source/Editor/GameProjectGeneration/Private/SNewClassDialog.cpp
+++ b/Engine/Source/Editor/GameProjectGeneration/Private/SNewClassDialog.cpp
@@ -312,7 +312,7 @@ void SNewClassDialog::Construct( const FArguments& InArgs )
 
 						+SHorizontalBox::Slot()
 						.VAlign(VAlign_Center)
-						.Padding(2.f)
+						.Padding(5.f, 2.f)
 						.AutoWidth()
 						[
 							SNew(SImage)
