Add log to more easily know why project file requires update
#rnx

diff --git a/Engine/Source/Editor/GameProjectGeneration/Private/GameProjectUtils.cpp b/Engine/Source/Editor/GameProjectGeneration/Private/GameProjectUtils.cpp
--- a/Engine/Source/Editor/GameProjectGeneration/Private/GameProjectUtils.cpp
+++ b/Engine/Source/Editor/GameProjectGeneration/Private/GameProjectUtils.cpp
@@ -915,6 +915,7 @@ void GameProjectUtils::CheckForOutOfDateGameProjectFile()
 			if ( ProjectStatus.bRequiresUpdate )
 			{
 				bRequiresUpdate = true;
+				UE_LOG(LogGameProjectGeneration, Log, TEXT("Project %s requires update. Determined from QueryStatusForCurrentProject"), *ProjectStatus.Name);
 			}
 		}
 
@@ -930,6 +931,7 @@ void GameProjectUtils::CheckForOutOfDateGameProjectFile()
 				FPluginReferenceDescriptor PluginReference(Plugin->GetName(), true);
 				NewPluginReferences.Add(PluginReference);
 				bRequiresUpdate = true;
+				UE_LOG(LogGameProjectGeneration, Log, TEXT("Project %s requires update. Installed plugin %s not found in project descriptor"), *ProjectStatus.Name, *Plugin->GetName());
 			}
 		}
 
@@ -944,11 +946,32 @@ void GameProjectUtils::CheckForOutOfDateGameProjectFile()
 					const FPluginDescriptor& Descriptor = Plugin->GetDescriptor();
 					if(Reference.MarketplaceURL != Descriptor.MarketplaceURL)
 					{
+						UE_LOG(LogGameProjectGeneration, Log, TEXT("Project %s requires update. Plugin %s MarketplaceURL value in project descriptor (%s) differs from value in plugin descriptor (%s) "),
+							*ProjectStatus.Name, *Plugin->GetName(), *Reference.MarketplaceURL, *Descriptor.MarketplaceURL);
+
 						Reference.MarketplaceURL = Descriptor.MarketplaceURL;
 						bRequiresUpdate = true;
 					}
 					if(Reference.SupportedTargetPlatforms != Descriptor.SupportedTargetPlatforms)
 					{
+						auto CombineStrings = [](const TArray<FString>& Strings) -> FString
+							{
+								FStringBuilderBase Output;
+								const int32 LastIndex = Strings.Num() - 1;
+								for (int32 I = 0; I <= LastIndex; ++I)
+								{
+									Output += Strings[I];
+									if (I != LastIndex)
+									{
+										Output += TEXT(", ");
+									}
+								}
+								return Output.ToString();
+							};
+
+						UE_LOG(LogGameProjectGeneration, Log, TEXT("Project %s requires update. Plugin %s SupportedTargetPlatforms value in project descriptor (%s) differs from value in plugin descriptor (%s) "),
+							*ProjectStatus.Name, *Plugin->GetName(), *CombineStrings(Reference.SupportedTargetPlatforms), *CombineStrings(Descriptor.SupportedTargetPlatforms));
+
 						Reference.SupportedTargetPlatforms = Descriptor.SupportedTargetPlatforms;
 						bRequiresUpdate = true;
 					}
