Add a "Don't show this again" checkbox to the "Project file is out of date." pop up at the start of the editor.

This will simply check the value of this cvar:
   [GameProjectUtils]
   SuppressUpdateProjectFilePromptNotification=true

Upon first clicking the "Don't show this again" checkbox, that CVar will be set to true.

#rnx

diff --git a/Engine/Source/Editor/GameProjectGeneration/Private/GameProjectUtils.cpp b/Engine/Source/Editor/GameProjectGeneration/Private/GameProjectUtils.cpp
--- a/Engine/Source/Editor/GameProjectGeneration/Private/GameProjectUtils.cpp
+++ b/Engine/Source/Editor/GameProjectGeneration/Private/GameProjectUtils.cpp
@@ -992,8 +992,39 @@ void GameProjectUtils::CheckForOutOfDateGameProjectFile()
 			const FText UpdateProjectConfirmText = LOCTEXT("UpdateProjectFileConfirm", "Update");
 			const FText UpdateProjectCancelText = LOCTEXT("UpdateProjectFileCancel", "Not Now");
 
+
+			/** Utility functions for notifications */
+			struct FSuppressDialogOptions
+			{
+				static bool ShouldSuppressModal()
+				{
+					bool bSuppressNotification = false;
+					GConfig->GetBool(TEXT("GameProjectUtils"), TEXT("SuppressUpdateProjectFilePromptNotification"), bSuppressNotification, GEditorPerProjectIni);
+					return bSuppressNotification;
+				}
+
+				static ECheckBoxState GetDontAskAgainCheckBoxState()
+				{
+					// Check the config for any preferences			
+					return ShouldSuppressModal() ? ECheckBoxState::Checked : ECheckBoxState::Unchecked;
+				}
+
+				static void OnDontAskAgainCheckBoxStateChanged(ECheckBoxState NewState)
+				{
+					// If the user selects to not show this again, set that in the config so we know about it in between sessions
+					const bool bSuppressNotification = (NewState == ECheckBoxState::Checked);
+					GConfig->SetBool(TEXT("GameProjectUtils"), TEXT("SuppressUpdateProjectFilePromptNotification"), bSuppressNotification, GEditorPerProjectIni);
+				}
+			}; 
+
+			// If the user has specified to supress this pop up, then just early out and exit	
+			if (FSuppressDialogOptions::ShouldSuppressModal())
+			{
+				return;
+			}
+			
 			FNotificationInfo Info(UpdateProjectText);
-			Info.ExpireDuration = 10;
+			Info.ExpireDuration = 10.0f;
 			Info.bFireAndForget = true;
 			Info.bUseLargeFont = false;
 			Info.bUseThrobber = false;
@@ -1001,6 +1032,11 @@ void GameProjectUtils::CheckForOutOfDateGameProjectFile()
 			Info.ButtonDetails.Add(FNotificationButtonInfo(UpdateProjectConfirmText, FText(), OnUpdateProjectConfirm));
 			Info.ButtonDetails.Add(FNotificationButtonInfo(UpdateProjectCancelText, FText(), FSimpleDelegate::CreateStatic(&GameProjectUtils::OnUpdateProjectCancel)));
 
+			// Add a "Don't show this again" option
+			Info.CheckBoxState = TAttribute<ECheckBoxState>::Create(&FSuppressDialogOptions::GetDontAskAgainCheckBoxState);
+			Info.CheckBoxStateChanged = FOnCheckStateChanged::CreateStatic(&FSuppressDialogOptions::OnDontAskAgainCheckBoxStateChanged);
+			Info.CheckBoxText = LOCTEXT("DefaultCheckBoxMessage", "Don't show this again");
+			
 			if (UpdateGameProjectNotification.IsValid())
 			{
 				UpdateGameProjectNotification.Pin()->ExpireAndFadeout();
