[GameProjectGeneration] Scaling down the New Project from Template selection window results in overlapping

diff --git a/Engine/Source/Editor/GameProjectGeneration/Private/SProjectDialog.cpp b/Engine/Source/Editor/GameProjectGeneration/Private/SProjectDialog.cpp
--- a/Engine/Source/Editor/GameProjectGeneration/Private/SProjectDialog.cpp
+++ b/Engine/Source/Editor/GameProjectGeneration/Private/SProjectDialog.cpp
@@ -691,6 +691,7 @@ TSharedRef<SWidget> SProjectDialog::MakeTemplateProjectView()
 			.AutoWidth()
 			[
 				SNew(SVerticalBox)
+				.Clipping(EWidgetClipping::ClipToBounds)
 				// Preview image
 				+ SVerticalBox::Slot()
 				.AutoHeight()
