Disable r.AllowStaticLighting by default in new projects

diff --git a/Engine/Source/Editor/GameProjectGeneration/Private/GameProjectUtils.cpp b/Engine/Source/Editor/GameProjectGeneration/Private/GameProjectUtils.cpp
--- a/Engine/Source/Editor/GameProjectGeneration/Private/GameProjectUtils.cpp
+++ b/Engine/Source/Editor/GameProjectGeneration/Private/GameProjectUtils.cpp
@@ -98,6 +98,13 @@ namespace
 	/** Get the configuration values for enabling Lumen by default. */
 	void AddLumenConfigValues(const FProjectInformation& InProjectInfo, TArray<FTemplateConfigValue>& ConfigValues)
 	{
+		// Disable static lighting support
+		ConfigValues.Emplace(TEXT("DefaultEngine.ini"),
+			TEXT("/Script/Engine.RendererSettings"),
+			TEXT("r.AllowStaticLighting"),
+			TEXT("0"),
+			true /* ShouldReplaceExistingValue */);
+
 		// Required for Lumen's Software Ray Tracing support
 		ConfigValues.Emplace(TEXT("DefaultEngine.ini"),
 			TEXT("/Script/Engine.RendererSettings"),
