[ProjectBrowser] Change "Most Recently Used" sorting to descending instead of ascending so that recent projects get listed first.

diff --git a/Engine/Source/Editor/GameProjectGeneration/Private/SProjectBrowser.cpp b/Engine/Source/Editor/GameProjectGeneration/Private/SProjectBrowser.cpp
--- a/Engine/Source/Editor/GameProjectGeneration/Private/SProjectBrowser.cpp
+++ b/Engine/Source/Editor/GameProjectGeneration/Private/SProjectBrowser.cpp
@@ -1339,7 +1339,7 @@ void SProjectBrowser::SortProjectTiles(EProjectSortOption NewSortOption)
 		break;
 	case EProjectSortOption::LastAccessTime:
 	default:
-		ProjectItemsSource.StableSort([](auto& A, auto& B) { return A->LastAccessTime < B->LastAccessTime; });
+		ProjectItemsSource.StableSort([](auto& A, auto& B) { return A->LastAccessTime > B->LastAccessTime; });
 		break;
 	}
 
