Updating templates that weren't setting r.AllowStaticLighting to either False or True, to r.AllowStaticLighting=False

diff --git a/Engine/Source/Editor/GameProjectGeneration/Private/GameProjectUtils.cpp b/Engine/Source/Editor/GameProjectGeneration/Private/GameProjectUtils.cpp
--- a/Engine/Source/Editor/GameProjectGeneration/Private/GameProjectUtils.cpp
+++ b/Engine/Source/Editor/GameProjectGeneration/Private/GameProjectUtils.cpp
@@ -96,14 +96,18 @@ constexpr const TCHAR GameProjectUtils::IncludePathFormatString[];
 namespace
 {
 	/** Get the configuration values for enabling Lumen by default. */
-	void AddLumenConfigValues(const FProjectInformation& InProjectInfo, TArray<FTemplateConfigValue>& ConfigValues)
+	void AddLumenConfigValues(const FProjectInformation& InProjectInfo, TArray<FTemplateConfigValue>& ConfigValues, bool bDisableStaticLighting = true)
 	{
-		// Disable static lighting support
-		ConfigValues.Emplace(TEXT("DefaultEngine.ini"),
-			TEXT("/Script/Engine.RendererSettings"),
-			TEXT("r.AllowStaticLighting"),
-			TEXT("0"),
-			true /* ShouldReplaceExistingValue */);
+		// In some cases we may want to leave static lighting enabled (e.g. VRTemplate needs it as of UE 5.5)
+		if (bDisableStaticLighting)
+		{
+			// Disable static lighting support
+			ConfigValues.Emplace(TEXT("DefaultEngine.ini"),
+				TEXT("/Script/Engine.RendererSettings"),
+				TEXT("r.AllowStaticLighting"),
+				TEXT("0"),
+				true /* ShouldReplaceExistingValue */);
+		}
 
 		// Required for Lumen's Software Ray Tracing support
 		ConfigValues.Emplace(TEXT("DefaultEngine.ini"),
@@ -1877,7 +1881,7 @@ TOptional<FGuid> GameProjectUtils::CreateProjectFromTemplate(const FProjectInfor
 
 	AddHardwareConfigValues(InProjectInfo, ConfigValuesToSet);
 
-	AddLumenConfigValues(InProjectInfo, ConfigValuesToSet);
+	AddLumenConfigValues(InProjectInfo, ConfigValuesToSet, /* do not disable static lighting, as some templates need it */ false);
 	AddRaytracingConfigValues(InProjectInfo, ConfigValuesToSet);
 	AddNewProjectDefaultShadowConfigValues(InProjectInfo, ConfigValuesToSet);
 	AddPostProcessingConfigValues(InProjectInfo, ConfigValuesToSet);
diff --git a/Templates/TP_AEC_CollabBP/Config/DefaultEngine.ini b/Templates/TP_AEC_CollabBP/Config/DefaultEngine.ini
--- a/Templates/TP_AEC_CollabBP/Config/DefaultEngine.ini
+++ b/Templates/TP_AEC_CollabBP/Config/DefaultEngine.ini
@@ -82,6 +82,7 @@ GameInstanceClass=/Game/CollaborativeViewer/Blueprints/GameInstance/BP_Collabora
 [/Script/Engine.RendererSettings]
 r.CustomDepth=3
 r.Mobile.UseHWsRGBEncoding=True
+r.AllowStaticLighting=False
 
 [/Script/Engine.PhysicsSettings]
 DefaultGravityZ=-980.000000
diff --git a/Templates/TP_AEC_ProdConfigBP/Config/DefaultEngine.ini b/Templates/TP_AEC_ProdConfigBP/Config/DefaultEngine.ini
--- a/Templates/TP_AEC_ProdConfigBP/Config/DefaultEngine.ini
+++ b/Templates/TP_AEC_ProdConfigBP/Config/DefaultEngine.ini
@@ -13,6 +13,7 @@ AppliedDefaultGraphicsPerformance=Maximum
 [/Script/Engine.RendererSettings]
 r.DefaultFeature.AutoExposure.ExtendDefaultLuminanceRange=True
 r.SupportSkyAtmosphereAffectsHeightFog=True
+r.AllowStaticLighting=False
 
 [/Script/EngineSettings.GameMapsSettings]
 EditorStartupMap=/Game/ProductAssets/Levels/Main.Main
diff --git a/Templates/TP_Blank/Config/DefaultEngine.ini b/Templates/TP_Blank/Config/DefaultEngine.ini
--- a/Templates/TP_Blank/Config/DefaultEngine.ini
+++ b/Templates/TP_Blank/Config/DefaultEngine.ini
@@ -3,6 +3,9 @@
 [/Script/EngineSettings.GameMapsSettings]
 GameDefaultMap=/Engine/Maps/Templates/OpenWorld
 
+[/Script/Engine.RendererSettings]
+r.AllowStaticLighting=False
+
 [/Script/WindowsTargetPlatform.WindowsTargetSettings]
 DefaultGraphicsRHI=DefaultGraphicsRHI_DX12
 -D3D12TargetedShaderFormats=PCD3D_SM5
diff --git a/Templates/TP_BlankBP/Config/DefaultEngine.ini b/Templates/TP_BlankBP/Config/DefaultEngine.ini
--- a/Templates/TP_BlankBP/Config/DefaultEngine.ini
+++ b/Templates/TP_BlankBP/Config/DefaultEngine.ini
@@ -3,6 +3,9 @@
 [/Script/EngineSettings.GameMapsSettings]
 GameDefaultMap=/Engine/Maps/Templates/OpenWorld
 
+[/Script/Engine.RendererSettings]
+r.AllowStaticLighting=False
+
 [/Script/WindowsTargetPlatform.WindowsTargetSettings]
 DefaultGraphicsRHI=DefaultGraphicsRHI_DX12
 -D3D12TargetedShaderFormats=PCD3D_SM5
diff --git a/Templates/TP_FirstPerson/Config/DefaultEngine.ini b/Templates/TP_FirstPerson/Config/DefaultEngine.ini
--- a/Templates/TP_FirstPerson/Config/DefaultEngine.ini
+++ b/Templates/TP_FirstPerson/Config/DefaultEngine.ini
@@ -21,6 +21,7 @@ r.ReflectionMethod=1
 r.GenerateMeshDistanceFields=True
 r.DynamicGlobalIlluminationMethod=1
 r.DefaultFeature.AutoExposure.ExtendDefaultLuminanceRange=true
+r.AllowStaticLighting=False
 
 
 [/Script/WindowsTargetPlatform.WindowsTargetSettings]
diff --git a/Templates/TP_FirstPersonBP/Config/DefaultEngine.ini b/Templates/TP_FirstPersonBP/Config/DefaultEngine.ini
--- a/Templates/TP_FirstPersonBP/Config/DefaultEngine.ini
+++ b/Templates/TP_FirstPersonBP/Config/DefaultEngine.ini
@@ -21,6 +21,7 @@ r.ReflectionMethod=1
 r.GenerateMeshDistanceFields=True
 r.DynamicGlobalIlluminationMethod=1
 r.DefaultFeature.AutoExposure.ExtendDefaultLuminanceRange=true
+r.AllowStaticLighting=False
 
 [/Script/WindowsTargetPlatform.WindowsTargetSettings]
 DefaultGraphicsRHI=DefaultGraphicsRHI_DX12
diff --git a/Templates/TP_InCamVFXBP/Config/DefaultEngine.ini b/Templates/TP_InCamVFXBP/Config/DefaultEngine.ini
--- a/Templates/TP_InCamVFXBP/Config/DefaultEngine.ini
+++ b/Templates/TP_InCamVFXBP/Config/DefaultEngine.ini
@@ -20,6 +20,7 @@ r.DefaultFeature.MotionBlur=False
 r.SceneRenderTargetResizeMethod=2
 r.DefaultBackBufferPixelFormat=4
 vr.AllowMotionBlurInVR=True
+r.AllowStaticLighting=False
 
 [/Script/WindowsTargetPlatform.WindowsTargetSettings]
 DefaultGraphicsRHI=DefaultGraphicsRHI_DX12
diff --git a/Templates/TP_MFG_CollabBP/Config/DefaultEngine.ini b/Templates/TP_MFG_CollabBP/Config/DefaultEngine.ini
--- a/Templates/TP_MFG_CollabBP/Config/DefaultEngine.ini
+++ b/Templates/TP_MFG_CollabBP/Config/DefaultEngine.ini
@@ -79,6 +79,7 @@ GameInstanceClass=/Game/CollaborativeViewer/Blueprints/GameInstance/BP_Collabora
 
 [/Script/Engine.RendererSettings]
 r.CustomDepth=3
+r.AllowStaticLighting=False
 
 [/Script/Engine.PhysicsSettings]
 DefaultGravityZ=-980.000000
diff --git a/Templates/TP_MFG_ProdConfigBP/Config/DefaultEngine.ini b/Templates/TP_MFG_ProdConfigBP/Config/DefaultEngine.ini
--- a/Templates/TP_MFG_ProdConfigBP/Config/DefaultEngine.ini
+++ b/Templates/TP_MFG_ProdConfigBP/Config/DefaultEngine.ini
@@ -16,6 +16,7 @@ r.DefaultFeature.AutoExposure=False
 r.DefaultFeature.LocalExposure.HighlightContrastScale=1.0
 r.DefaultFeature.LocalExposure.ShadowContrastScale=1.0
 r.GenerateMeshDistanceFields=True
+r.AllowStaticLighting=False
 
 [/Script/EngineSettings.GameMapsSettings]
 EditorStartupMap=/Game/ProductAssets/Levels/Main.Main
diff --git a/Templates/TP_SIM_SandboxBP/Config/DefaultEngine.ini b/Templates/TP_SIM_SandboxBP/Config/DefaultEngine.ini
--- a/Templates/TP_SIM_SandboxBP/Config/DefaultEngine.ini
+++ b/Templates/TP_SIM_SandboxBP/Config/DefaultEngine.ini
@@ -14,6 +14,7 @@ r.Shadow.Virtual.Enable=1
 r.GenerateMeshDistanceFields=True
 r.DynamicGlobalIlluminationMethod=1
 r.ReflectionMethod=1
+r.AllowStaticLighting=False
 
 [/Script/WindowsTargetPlatform.WindowsTargetSettings]
 DefaultGraphicsRHI=DefaultGraphicsRHI_DX12
diff --git a/Templates/TP_ThirdPerson/Config/DefaultEngine.ini b/Templates/TP_ThirdPerson/Config/DefaultEngine.ini
--- a/Templates/TP_ThirdPerson/Config/DefaultEngine.ini
+++ b/Templates/TP_ThirdPerson/Config/DefaultEngine.ini
@@ -10,6 +10,7 @@ r.DynamicGlobalIlluminationMethod=1
 r.Lumen.TraceMeshSDFs=0
 r.Shadow.Virtual.Enable=1
 r.DefaultFeature.AutoExposure.ExtendDefaultLuminanceRange=true
+r.AllowStaticLighting=False
 
 [/Script/WindowsTargetPlatform.WindowsTargetSettings]
 DefaultGraphicsRHI=DefaultGraphicsRHI_DX12
diff --git a/Templates/TP_ThirdPersonBP/Config/DefaultEngine.ini b/Templates/TP_ThirdPersonBP/Config/DefaultEngine.ini
--- a/Templates/TP_ThirdPersonBP/Config/DefaultEngine.ini
+++ b/Templates/TP_ThirdPersonBP/Config/DefaultEngine.ini
@@ -18,6 +18,7 @@ r.DynamicGlobalIlluminationMethod=1
 r.Lumen.TraceMeshSDFs=0
 r.Shadow.Virtual.Enable=1
 r.DefaultFeature.AutoExposure.ExtendDefaultLuminanceRange=true
+r.AllowStaticLighting=False
 
 [/Script/WindowsTargetPlatform.WindowsTargetSettings]
 DefaultGraphicsRHI=DefaultGraphicsRHI_DX12
diff --git a/Templates/TP_TopDown/Config/DefaultEngine.ini b/Templates/TP_TopDown/Config/DefaultEngine.ini
--- a/Templates/TP_TopDown/Config/DefaultEngine.ini
+++ b/Templates/TP_TopDown/Config/DefaultEngine.ini
@@ -5,6 +5,7 @@ GlobalDefaultGameMode="/Script/TP_TopDown.TP_TopDownGameMode"
 
 [/Script/Engine.RendererSettings]
 r.DefaultFeature.AutoExposure.ExtendDefaultLuminanceRange=true
+r.AllowStaticLighting=False
 
 [/Script/WindowsTargetPlatform.WindowsTargetSettings]
 DefaultGraphicsRHI=DefaultGraphicsRHI_DX12
diff --git a/Templates/TP_TopDownBP/Config/DefaultEngine.ini b/Templates/TP_TopDownBP/Config/DefaultEngine.ini
--- a/Templates/TP_TopDownBP/Config/DefaultEngine.ini
+++ b/Templates/TP_TopDownBP/Config/DefaultEngine.ini
@@ -55,6 +55,7 @@ r.GenerateMeshDistanceFields=True
 r.DynamicGlobalIlluminationMethod=1
 r.Shadow.Virtual.Enable=1
 r.DefaultFeature.AutoExposure.ExtendDefaultLuminanceRange=true
+r.AllowStaticLighting=False
 
 [/Script/WindowsTargetPlatform.WindowsTargetSettings]
 DefaultGraphicsRHI=DefaultGraphicsRHI_DX12
diff --git a/Templates/TP_VehicleAdv/Config/DefaultEngine.ini b/Templates/TP_VehicleAdv/Config/DefaultEngine.ini
--- a/Templates/TP_VehicleAdv/Config/DefaultEngine.ini
+++ b/Templates/TP_VehicleAdv/Config/DefaultEngine.ini
@@ -19,6 +19,7 @@ GlobalDefaultServerGameMode=None
 r.ReflectionMethod=1
 r.GenerateMeshDistanceFields=True
 r.DynamicGlobalIlluminationMethod=1
+r.AllowStaticLighting=False
 
 [/Script/WindowsTargetPlatform.WindowsTargetSettings]
 DefaultGraphicsRHI=DefaultGraphicsRHI_DX12
diff --git a/Templates/TP_VehicleAdvBP/Config/DefaultEngine.ini b/Templates/TP_VehicleAdvBP/Config/DefaultEngine.ini
--- a/Templates/TP_VehicleAdvBP/Config/DefaultEngine.ini
+++ b/Templates/TP_VehicleAdvBP/Config/DefaultEngine.ini
@@ -24,6 +24,7 @@ r.GenerateMeshDistanceFields=True
 r.DynamicGlobalIlluminationMethod=1
 r.ClearCoatNormal=True
 r.DefaultFeature.AutoExposure.ExtendDefaultLuminanceRange=true
+r.AllowStaticLighting=False
 
 [/Script/WindowsTargetPlatform.WindowsTargetSettings]
 DefaultGraphicsRHI=DefaultGraphicsRHI_DX12
diff --git a/Templates/TP_nDisplayBP/Config/DefaultEngine.ini b/Templates/TP_nDisplayBP/Config/DefaultEngine.ini
--- a/Templates/TP_nDisplayBP/Config/DefaultEngine.ini
+++ b/Templates/TP_nDisplayBP/Config/DefaultEngine.ini
@@ -21,6 +21,7 @@ r.SceneRenderTargetResizeMethod=2
 vr.AllowMotionBlurInVR=True
 r.DefaultFeature.AutoExposure.Method=1
 r.DefaultFeature.AutoExposure.Bias=1.000000
+r.AllowStaticLighting=False
 
 [/Script/WindowsTargetPlatform.WindowsTargetSettings]
 DefaultGraphicsRHI=DefaultGraphicsRHI_DX12
