SNewClassDialog: Add an argument to configure whether the newly created class is focused in the Content Browser

diff --git a/Engine/Source/Editor/GameProjectGeneration/Private/GameProjectUtils.cpp b/Engine/Source/Editor/GameProjectGeneration/Private/GameProjectUtils.cpp
--- a/Engine/Source/Editor/GameProjectGeneration/Private/GameProjectUtils.cpp
+++ b/Engine/Source/Editor/GameProjectGeneration/Private/GameProjectUtils.cpp
@@ -1187,7 +1187,8 @@ void GameProjectUtils::OpenAddToProjectDialog(const FAddToProjectConfig& Config,
 		.InitialPath(Config._InitialPath)
 		.OnAddedToProject( Config._OnAddedToProject )
 		.DefaultClassPrefix( Config._DefaultClassPrefix )
-		.DefaultClassName( Config._DefaultClassName );
+		.DefaultClassName( Config._DefaultClassName )
+		.SyncContentBrowserToNewClass(Config._SyncContentBrowserToNewClass);
 
 	AddCodeWindow->SetContent( NewClassDialog );
 
diff --git a/Engine/Source/Editor/GameProjectGeneration/Private/SNewClassDialog.cpp b/Engine/Source/Editor/GameProjectGeneration/Private/SNewClassDialog.cpp
--- a/Engine/Source/Editor/GameProjectGeneration/Private/SNewClassDialog.cpp
+++ b/Engine/Source/Editor/GameProjectGeneration/Private/SNewClassDialog.cpp
@@ -100,6 +100,7 @@ BEGIN_SLATE_FUNCTION_BUILD_OPTIMIZATION
 void SNewClassDialog::Construct( const FArguments& InArgs )
 {
 	ClassDomain = InArgs._ClassDomain;
+	bSyncContentBrowserToNewClass = InArgs._SyncContentBrowserToNewClass;
 
 	{
 		TArray<FModuleContextInfo> CurrentModules = GameProjectUtils::GetCurrentProjectModules();
@@ -1164,8 +1165,11 @@ void SNewClassDialog::FinishClicked()
 
 					OnAddedToProject.ExecuteIfBound( NewClassName, PackagePath, FString() );
 
-					// Sync the content browser to the new asset
-					GEditor->SyncBrowserToObject(NewBP);
+					if (bSyncContentBrowserToNewClass)
+					{
+						// Sync the content browser to the new asset
+						GEditor->SyncBrowserToObject(NewBP);
+					}
 
 					// Open the editor for the new asset
 					GEditor->GetEditorSubsystem<UAssetEditorSubsystem>()->OpenEditorForAsset(NewBP);
diff --git a/Engine/Source/Editor/GameProjectGeneration/Private/SNewClassDialog.h b/Engine/Source/Editor/GameProjectGeneration/Private/SNewClassDialog.h
--- a/Engine/Source/Editor/GameProjectGeneration/Private/SNewClassDialog.h
+++ b/Engine/Source/Editor/GameProjectGeneration/Private/SNewClassDialog.h
@@ -49,7 +49,9 @@ class SNewClassDialog : public SCompoundWidget
 {
 public:
 	SLATE_BEGIN_ARGS( SNewClassDialog )
-		: _ClassDomain(EClassDomain::Native), _Class(NULL)
+		: _ClassDomain(EClassDomain::Native)
+		, _Class(NULL)
+		, _SyncContentBrowserToNewClass(true)
 	{}
 
 	/** A reference to the parent window */
@@ -80,6 +82,9 @@ public:
 	/** Event called when code is successfully added to the project */
 	SLATE_EVENT( FOnAddedToProject, OnAddedToProject )
 
+	/** Whether we should sync the content browser to focus on the newly created asset */
+	SLATE_ARGUMENT(bool, SyncContentBrowserToNewClass)
+
 	SLATE_END_ARGS()
 
 	/** Constructs this widget with InArgs */
@@ -301,4 +306,7 @@ private:
 
 	/** Flag set when PathPicker selects a path*/
 	bool IsBlueprintPathSelected = false;
+
+	/** Whether we should sync the content browser to navigate to the newly created asset */
+	bool bSyncContentBrowserToNewClass = true;
 };
diff --git a/Engine/Source/Editor/GameProjectGeneration/Public/AddToProjectConfig.h b/Engine/Source/Editor/GameProjectGeneration/Public/AddToProjectConfig.h
--- a/Engine/Source/Editor/GameProjectGeneration/Public/AddToProjectConfig.h
+++ b/Engine/Source/Editor/GameProjectGeneration/Public/AddToProjectConfig.h
@@ -179,6 +179,8 @@ public:
 	FAddToProjectConfig& FeatureActorClasses()										{ _FeaturedClasses = FFeaturedClasses::ActorClasses(); return *this; }
 	/** Set the add to project dialog to show component types on the initial 'featured' classes list */
 	FAddToProjectConfig& FeatureComponentClasses()									{ _FeaturedClasses = FFeaturedClasses::ComponentClasses(); return *this; }
+	/** Set whether the content browser will be synced to view the newly created class */
+	FAddToProjectConfig& SyncContentBrowserToNewClass(bool bInSyncContentBrowser)	{ _SyncContentBrowserToNewClass = bInSyncContentBrowser; return *this; }
 
 public:
 
@@ -205,4 +207,7 @@ public:
 
 	/** Delegate to invoke when the (blueprint) class has been added to the project */
 	FOnAddedToProject _OnAddedToProject;
+
+	/** Whether the content browser will be synced to view the newly created class */
+	bool _SyncContentBrowserToNewClass = true;
 };
