NewClassDialog: Allow users to create new folders in the path picker by enabling the context menu

diff --git a/Engine/Source/Editor/GameProjectGeneration/Private/SNewClassDialog.cpp b/Engine/Source/Editor/GameProjectGeneration/Private/SNewClassDialog.cpp
--- a/Engine/Source/Editor/GameProjectGeneration/Private/SNewClassDialog.cpp
+++ b/Engine/Source/Editor/GameProjectGeneration/Private/SNewClassDialog.cpp
@@ -268,7 +268,7 @@ void SNewClassDialog::Construct( const FArguments& InArgs )
 	{
 		BlueprintPathConfig.DefaultPath = NewClassPath;
 		BlueprintPathConfig.bFocusSearchBoxWhenOpened = false;
-		BlueprintPathConfig.bAllowContextMenu = false;
+		BlueprintPathConfig.bAllowContextMenu = true;
 		BlueprintPathConfig.bAllowClassesFolder = false;
 		BlueprintPathConfig.bAllowReadOnlyFolders = false;
 		BlueprintPathConfig.OnPathSelected = FOnPathSelected::CreateSP(this, &SNewClassDialog::OnBlueprintPathSelected);
