CommonUI: Add UCommonUISettings::CommonButtonAcceptKeyHandling to configure how SCommonButton reacts to key events for FNavigationConfig's Accept action

For existing projects, the default is to ignore those key events and to rely on the button's Triggering Input Action instead. New projects will use the "TriggerClick" option by default and it is recommended to update the FNavigationConfig appropriately instead of switching to "Ignore" mode.

#rnx

PR #12164
PR #12212

diff --git a/Engine/Plugins/Runtime/CommonUI/Source/CommonUI/Private/CommonButtonTypes.cpp b/Engine/Plugins/Runtime/CommonUI/Source/CommonUI/Private/CommonButtonTypes.cpp
--- a/Engine/Plugins/Runtime/CommonUI/Source/CommonUI/Private/CommonButtonTypes.cpp
+++ b/Engine/Plugins/Runtime/CommonUI/Source/CommonUI/Private/CommonButtonTypes.cpp
@@ -1,6 +1,10 @@
 // Copyright Epic Games, Inc. All Rights Reserved.
 
 #include "CommonButtonTypes.h"
+
+#include "CommonUISettings.h"
+#include "Framework/Application/SlateApplication.h"
+#include "ICommonUIModule.h"
 #include "Input/Reply.h"
 
 //////////////////////////////////////////////////////////////////////////
@@ -129,7 +133,8 @@ FReply SCommonButton::OnTouchMoved(const FGeometry& MyGeometry, const FPointerEv
 
 FReply SCommonButton::OnKeyDown(const FGeometry& MyGeometry, const FKeyEvent& InKeyEvent)
 {
-	if (InKeyEvent.GetKey() == EKeys::Enter)
+	if (ICommonUIModule::GetSettings().GetCommonButtonAcceptKeyHandling() == ECommonButtonAcceptKeyHandling::Ignore
+		&& FSlateApplication::Get().GetNavigationActionFromKey(InKeyEvent) == EUINavigationAction::Accept)
 	{
 		return FReply::Unhandled();
 	}
@@ -138,7 +143,8 @@ FReply SCommonButton::OnKeyDown(const FGeometry& MyGeometry, const FKeyEvent& In
 
 FReply SCommonButton::OnKeyUp(const FGeometry& MyGeometry, const FKeyEvent& InKeyEvent)
 {
-	if (InKeyEvent.GetKey() == EKeys::Enter)
+	if (ICommonUIModule::GetSettings().GetCommonButtonAcceptKeyHandling() == ECommonButtonAcceptKeyHandling::Ignore
+		&& FSlateApplication::Get().GetNavigationActionFromKey(InKeyEvent) == EUINavigationAction::Accept)
 	{
 		return FReply::Unhandled();
 	}
diff --git a/Engine/Plugins/Runtime/CommonUI/Source/CommonUI/Private/CommonUISettings.cpp b/Engine/Plugins/Runtime/CommonUI/Source/CommonUI/Private/CommonUISettings.cpp
--- a/Engine/Plugins/Runtime/CommonUI/Source/CommonUI/Private/CommonUISettings.cpp
+++ b/Engine/Plugins/Runtime/CommonUI/Source/CommonUI/Private/CommonUISettings.cpp
@@ -13,6 +13,7 @@ UE_DEFINE_GAMEPLAY_TAG(TAG_PlatformTrait_PlayInEditor, "Platform.Trait.PlayInEdi
 UCommonUISettings::UCommonUISettings(const FObjectInitializer& Initializer)
 	: Super(Initializer)
 	, bAutoLoadData(true)
+	, CommonButtonAcceptKeyHandling(ECommonButtonAcceptKeyHandling::Ignore)
 	, bDefaultDataLoaded(false)
 {}
 
@@ -175,3 +176,8 @@ const FGameplayTagContainer& UCommonUISettings::GetPlatformTraits() const
 
 	return PlatformTraitContainer;
 }
+
+ECommonButtonAcceptKeyHandling UCommonUISettings::GetCommonButtonAcceptKeyHandling() const
+{
+	return CommonButtonAcceptKeyHandling;
+}
diff --git a/Engine/Plugins/Runtime/CommonUI/Source/CommonUI/Public/CommonUISettings.h b/Engine/Plugins/Runtime/CommonUI/Source/CommonUI/Public/CommonUISettings.h
--- a/Engine/Plugins/Runtime/CommonUI/Source/CommonUI/Public/CommonUISettings.h
+++ b/Engine/Plugins/Runtime/CommonUI/Source/CommonUI/Public/CommonUISettings.h
@@ -13,6 +13,20 @@ class UMaterialInterface;
 
 COMMONUI_API UE_DECLARE_GAMEPLAY_TAG_EXTERN(TAG_PlatformTrait_PlayInEditor);
 
+UENUM()
+enum class ECommonButtonAcceptKeyHandling
+{
+	// (Default for existing projects) When a CommonButton is focused, it will ignore the configured keys for slate navigation's Accept action (see FNavigationConfig::KeyActionRules),
+	// leaving it up to the button's Triggering Input Action to trigger the button's click event.
+	// This allows other objects to bind input actions to the accept action keys while focusing a common button.
+	Ignore,
+
+	// (Default for new projects) When a CommonButton is focused, it will let all inputs events flow through to the underlying SButton, therefore triggering the button's click event when
+	// pressing the configured keys for slate navigation's Accept action (see FNavigationConfig::KeyActionRules)
+	// This will prevent input actions that are bound to those keys from triggering while a CommonButton is focused.
+	TriggerClick,
+};
+
 UCLASS(config = Game, defaultconfig)
 class COMMONUI_API UCommonUISettings : public UObject
 {
@@ -41,6 +55,7 @@ public:
 	const FSlateBrush& GetDefaultThrobberBrush() const;
 	UObject* GetDefaultImageResourceObject() const;
 	const FGameplayTagContainer& GetPlatformTraits() const;
+	ECommonButtonAcceptKeyHandling GetCommonButtonAcceptKeyHandling() const;
 
 private:
 
@@ -66,6 +81,12 @@ private:
 	UPROPERTY(config, EditAnywhere, Category = "Visibility", meta=(Categories="Platform.Trait", ConfigHierarchyEditable))
 	TArray<FGameplayTag> PlatformTraits;
 
+	/**
+	 * How should CommonButton widgets handle SlateNavigation Accept actions?
+	 */
+	UPROPERTY(config, EditAnywhere, Category = "Button")
+	ECommonButtonAcceptKeyHandling CommonButtonAcceptKeyHandling;
+
 private:
 	void LoadEditorData();
 	void RebuildTraitContainer();
diff --git a/Engine/Source/Editor/GameProjectGeneration/Private/GameProjectUtils.cpp b/Engine/Source/Editor/GameProjectGeneration/Private/GameProjectUtils.cpp
--- a/Engine/Source/Editor/GameProjectGeneration/Private/GameProjectUtils.cpp
+++ b/Engine/Source/Editor/GameProjectGeneration/Private/GameProjectUtils.cpp
@@ -223,7 +223,7 @@ namespace
 		}
 	}
 
-	void AddUserInterfaceConfigValues(const FProjectInformation& InProjectInfo, TArray<FTemplateConfigValue>& ConfigValues)
+	void AddUserInterfaceEngineConfigValues(const FProjectInformation& InProjectInfo, TArray<FTemplateConfigValue>& ConfigValues)
 	{
 		if (InProjectInfo.bIsBlankTemplate)
 		{
@@ -246,6 +246,18 @@ namespace
 				true /* ShouldReplaceExistingValue */);
 		}
 	}
+
+	void AddUserInterfaceGameConfigValues(const FProjectInformation& InProjectInfo, TArray<FTemplateConfigValue>& ConfigValues)
+	{
+		if (InProjectInfo.bIsBlankTemplate)
+		{
+			ConfigValues.Emplace(TEXT("DefaultGame.ini"),
+				TEXT("/Script/CommonUI.CommonUISettings"),
+				TEXT("CommonButtonAcceptKeyHandling"),
+				TEXT("TriggerClick"),
+				true /* ShouldReplaceExistingValue */);
+		}
+	}
 } // namespace <>
 
 FText FNewClassInfo::GetClassName() const
@@ -1887,7 +1899,8 @@ TOptional<FGuid> GameProjectUtils::CreateProjectFromTemplate(const FProjectInfor
 	AddNewProjectDefaultShadowConfigValues(InProjectInfo, ConfigValuesToSet);
 	AddPostProcessingConfigValues(InProjectInfo, ConfigValuesToSet);
 	AddWorldPartitionConfigValues(InProjectInfo, ConfigValuesToSet);
-	AddUserInterfaceConfigValues(InProjectInfo, ConfigValuesToSet);
+	AddUserInterfaceEngineConfigValues(InProjectInfo, ConfigValuesToSet);
+	AddUserInterfaceGameConfigValues(InProjectInfo, ConfigValuesToSet);
 	
 	TemplateDefs->AddConfigValues(ConfigValuesToSet, TemplateName, ProjectName, InProjectInfo.bShouldGenerateCode);
 
@@ -2200,7 +2213,7 @@ bool GameProjectUtils::GenerateConfigFiles(const FProjectInformation& InProjectI
 		AddPostProcessingConfigValues(InProjectInfo, ConfigValuesToSet);
 		AddRaytracingConfigValues(InProjectInfo, ConfigValuesToSet);
 		AddWorldPartitionConfigValues(InProjectInfo, ConfigValuesToSet);
-		AddUserInterfaceConfigValues(InProjectInfo, ConfigValuesToSet);
+		AddUserInterfaceEngineConfigValues(InProjectInfo, ConfigValuesToSet);
 
 		if (!SaveConfigValues(InProjectInfo, ConfigValuesToSet, OutFailReason))
 		{
@@ -2240,6 +2253,14 @@ bool GameProjectUtils::GenerateConfigFiles(const FProjectInformation& InProjectI
 		{
 			return false;
 		}
+
+		TArray<FTemplateConfigValue> ConfigValuesToSet;
+		AddUserInterfaceGameConfigValues(InProjectInfo, ConfigValuesToSet);
+
+		if (!SaveConfigValues(InProjectInfo, ConfigValuesToSet, OutFailReason))
+		{
+			return false;
+		}
 	}
 
 	// Platforms inis:
