Don't automatically add installed plugins as project dependencies
- Forcibly adding a hard dependency the user didn't ask for can cause problems. We can't know for sure the user actually wants that dependency. One concrete problem this causes is that we have an internal test project that constantly thinks it's out of date and tries to update itself. When it does, it adds a random plugin that is not actually dependency of that project and that in turn breaks licensees that can see the test project but not the plugin. The plugin really is "installed" so we don't want to change that property and that leaves this as the only fix.

#ref CL 2566860

diff --git a/Engine/Source/Editor/GameProjectGeneration/Private/GameProjectUtils.cpp b/Engine/Source/Editor/GameProjectGeneration/Private/GameProjectUtils.cpp
--- a/Engine/Source/Editor/GameProjectGeneration/Private/GameProjectUtils.cpp
+++ b/Engine/Source/Editor/GameProjectGeneration/Private/GameProjectUtils.cpp
@@ -942,23 +942,9 @@ void GameProjectUtils::CheckForOutOfDateGameProjectFile()
 			}
 		}
 
-		// Get the current project descriptor
+		// Check if there are any referenced plugins that do not have a matching supported plugins list
 		const FProjectDescriptor* Project = IProjectManager::Get().GetCurrentProject();
-
-		// Check if there are any installed plugins that need to be added as a reference
 		TArray<FPluginReferenceDescriptor> NewPluginReferences = Project->Plugins;
-		for(TSharedRef<IPlugin>& Plugin: IPluginManager::Get().GetEnabledPlugins())
-		{
-			if(Plugin->GetDescriptor().bInstalled && Project->FindPluginReferenceIndex(Plugin->GetName()) == INDEX_NONE)
-			{
-				FPluginReferenceDescriptor PluginReference(Plugin->GetName(), true);
-				NewPluginReferences.Add(PluginReference);
-				bRequiresUpdate = true;
-				UE_LOG(LogGameProjectGeneration, Log, TEXT("Project %s requires update. Installed plugin %s not found in project descriptor"), *ProjectStatus.Name, *Plugin->GetName());
-			}
-		}
-
-		// Check if there are any referenced plugins that do not have a matching supported plugins list
 		for(FPluginReferenceDescriptor& Reference: NewPluginReferences)
 		{
 			if(Reference.bEnabled)
