PR #9737: Fix new class wizard default package

diff --git a/Engine/Source/Editor/GameProjectGeneration/Private/SNewClassDialog.cpp b/Engine/Source/Editor/GameProjectGeneration/Private/SNewClassDialog.cpp
--- a/Engine/Source/Editor/GameProjectGeneration/Private/SNewClassDialog.cpp
+++ b/Engine/Source/Editor/GameProjectGeneration/Private/SNewClassDialog.cpp
@@ -154,7 +154,7 @@ void SNewClassDialog::Construct( const FArguments& InArgs )
 		const FString ProjectName = FApp::GetProjectName();
 
 		// Find initially selected module based on simple fallback in this order..
-		// Previously selected module, main project module, a  runtime module
+		// Previously selected module, main project module, the first runtime module
 		TSharedPtr<FModuleContextInfo> ProjectModule;
 		TSharedPtr<FModuleContextInfo> RuntimeModule;
 
@@ -171,8 +171,9 @@ void SNewClassDialog::Construct( const FArguments& InArgs )
 			{
 				ProjectModule = AvailableModule;
 			}
-
-			if (AvailableModule->ModuleType == EHostType::Runtime)
+			
+			//Check if this is a runtime module and if we haven't already found one
+			if (AvailableModule->ModuleType == EHostType::Runtime && RuntimeModule == nullptr)
 			{
 				RuntimeModule = AvailableModule;
 			}
