Removing Starter Content References

Ability to add starter content to a new project has been deprecated and these changes remove references to it.

diff --git a/Engine/Config/BaseEditorSettings.ini b/Engine/Config/BaseEditorSettings.ini
--- a/Engine/Config/BaseEditorSettings.ini
+++ b/Engine/Config/BaseEditorSettings.ini
@@ -7,9 +7,6 @@ ShouldUseLocalizedNumericInput=True
 ShouldUseLocalizedPropertyNames=True
 ShouldUseLocalizedNodeAndPinNames=False
 
-[/Script/UnrealEd.EditorSettings]
-bCopyStarterContentPreference = true
-
 [/Script/UnrealEd.CrashReportsPrivacySettings]
 bSendUnattendedBugReports=True
 
diff --git a/Engine/Source/Editor/GameProjectGeneration/Classes/TemplateProjectDefs.h b/Engine/Source/Editor/GameProjectGeneration/Classes/TemplateProjectDefs.h
--- a/Engine/Source/Editor/GameProjectGeneration/Classes/TemplateProjectDefs.h
+++ b/Engine/Source/Editor/GameProjectGeneration/Classes/TemplateProjectDefs.h
@@ -80,7 +80,7 @@ enum class ETemplateSetting
 	Languages,
 	HardwareTarget,
 	GraphicsPreset,
-	StarterContent,
+	StarterContent UE_DEPRECATED(5.6, "Ability to add Starter Content has been removed"),
 	XR,
 	Raytracing_DEPRECATED,
 	All
@@ -157,8 +157,11 @@ public:
 	UPROPERTY(config)
 	TArray<FFeaturePackLevelSet> SharedContentPacks;
 
-	UPROPERTY(config)
-	FString StarterContent;
+#if WITH_EDITORONLY_DATA
+	UE_DEPRECATED(5.6, "Ability to add Starter Content has been removed")
+	UPROPERTY(config, meta = (DeprecatedProperty, DeprecationMessage = "Ability to add Starter Content has been removed"))
+	FString StarterContent_DEPRECATED;
+#endif // WITH_EDITORONLY_DATA
 
 	/** Fixes up all strings in this definitions object to replace \%TEMPLATENAME\% with the supplied template name and \%PROJECTNAME\% with the supplied project name */
 	void FixupStrings(const FString& TemplateName, const FString& ProjectName);
diff --git a/Engine/Source/Editor/GameProjectGeneration/Private/GameProjectUtils.cpp b/Engine/Source/Editor/GameProjectGeneration/Private/GameProjectUtils.cpp
--- a/Engine/Source/Editor/GameProjectGeneration/Private/GameProjectUtils.cpp
+++ b/Engine/Source/Editor/GameProjectGeneration/Private/GameProjectUtils.cpp
@@ -187,29 +187,6 @@ namespace
 			false /* ShouldReplaceExistingValue */);
 	}
 
-	void AddDefaultMapConfigValues(const FProjectInformation& InProjectInfo, TArray<FTemplateConfigValue>& ConfigValues)
-	{
-		if (InProjectInfo.bIsBlankTemplate &&
-			InProjectInfo.bCopyStarterContent &&
-			GameProjectUtils::IsUsingEngineStarterContent(InProjectInfo) &&
-			GameProjectUtils::IsEngineStarterContentAvailable() )
-		{
-			const FString DefaultMap = TEXT("/Game/StarterContent/Maps/Minimal_Default");
-
-			ConfigValues.Emplace(TEXT("DefaultEngine.ini"),
-				TEXT("/Script/EngineSettings.GameMapsSettings"),
-				TEXT("EditorStartupMap"),
-				DefaultMap,
-				true /* ShouldReplaceExistingValue */);
-
-			ConfigValues.Emplace(TEXT("DefaultEngine.ini"),
-				TEXT("/Script/EngineSettings.GameMapsSettings"),
-				TEXT("GameDefaultMap"),
-				DefaultMap,
-				true /* ShouldReplaceExistingValue */);
-		}
-	}
-
 	/** Get the configuration values for enabling WorldPartition by default. */
 	void AddWorldPartitionConfigValues(const FProjectInformation& InProjectInfo, TArray<FTemplateConfigValue>& ConfigValues)
 	{
@@ -815,36 +792,6 @@ bool GameProjectUtils::OpenCodeIDE(const FString& ProjectFile, FText& OutFailRea
 	return true;
 }
 
-bool GameProjectUtils::IsEngineStarterContentAvailable()
-{
-	TArray<FString> OutFilenames;
-	IFileManager::Get().FindFilesRecursive(OutFilenames, *FPaths::FeaturePackDir(), TEXT("*StarterContent.upack"), /*Files=*/true, /*Directories=*/false);
-	return OutFilenames.Num() > 0;
-}
-
-bool GameProjectUtils::IsUsingEngineStarterContent(const FProjectInformation& InProjectInfo)
-{
-	return InProjectInfo.StarterContent.IsEmpty();
-}
-
-FString GameProjectUtils::GetStarterContentName(const FProjectInformation& InProjectInfo)
-{
-	if (!InProjectInfo.StarterContent.IsEmpty())
-	{
-		return InProjectInfo.StarterContent;
-	}
-
-	return TEXT("StarterContent");
-}
-
-bool GameProjectUtils::IsStarterContentAvailableForProject(const FProjectInformation& InProjectInfo)
-{
-	const FString StarterContentName = GameProjectUtils::GetStarterContentName(InProjectInfo);
-	const FString StarterContentPackFilename = FPaths::FeaturePackDir() / FString::Printf(TEXT("%s.upack"), *StarterContentName);
-
-	return IFileManager::Get().FileExists(*StarterContentPackFilename);
-}
-
 bool GameProjectUtils::CreateProject(const FProjectInformation& InProjectInfo, FText& OutFailReason, FText& OutFailLog, TArray<FString>* OutCreatedFiles)
 {
 	if ( !IsValidProjectFileForCreation(InProjectInfo.ProjectFilename, OutFailReason) )
@@ -908,8 +855,6 @@ bool GameProjectUtils::CreateProject(const FProjectInformation& InProjectInfo, F
 			}
 		}
 
-		EventAttributes.Add(FAnalyticsEventAttribute(TEXT("StarterContent"), InProjectInfo.bCopyStarterContent ? TEXT("Yes") : TEXT("No")));
-
 		FEngineAnalytics::GetProvider().RecordEvent(TEXT("Editor.NewProject.ProjectCreated"), EventAttributes);
 	}
 
@@ -1451,12 +1396,6 @@ TOptional<FGuid> GameProjectUtils::GenerateProjectFromScratch(const FProjectInfo
 		return TOptional<FGuid>();
 	}
 
-	// Insert any required feature packs (EG starter content) into ini file. These will be imported automatically when the editor is first run
-	if(!InsertFeaturePacksIntoINIFile(InProjectInfo, OutFailReason))
-	{
-		return TOptional<FGuid>();
-	}
-
 	// Make the Content folder
 	const FString ContentFolder = NewProjectFolder / TEXT("Content");
 	if ( !IFileManager::Get().MakeDirectory(*ContentFolder) )
@@ -1907,12 +1846,6 @@ TOptional<FGuid> GameProjectUtils::CreateProjectFromTemplate(const FProjectInfor
 		return TOptional<FGuid>();
 	}
 
-	// Insert any required feature packs (EG starter content) into ini file. These will be imported automatically when the editor is first run
-	if (!InsertFeaturePacksIntoINIFile(InProjectInfo, OutFailReason))
-	{
-		return TOptional<FGuid>();
-	}
-
 	if (!AddSharedContentToProject(InProjectInfo, CreatedFiles, OutFailReason))
 	{
 		return TOptional<FGuid>();
@@ -2166,23 +2099,6 @@ bool GameProjectUtils::GenerateConfigFiles(const FProjectInformation& InProjectI
 		FileContents += TEXT("[Audio]") LINE_TERMINATOR;
 		FileContents += TEXT("UseAudioMixer=True") LINE_TERMINATOR;
 
-		if (InProjectInfo.bCopyStarterContent)
-		{
-			FileContents += LINE_TERMINATOR;
-			FileContents += TEXT("[/Script/EngineSettings.GameMapsSettings]") LINE_TERMINATOR;
-
-			if (GameProjectUtils::IsEngineStarterContentAvailable() && GameProjectUtils::IsUsingEngineStarterContent(InProjectInfo)) // if use Engine StarterContent
-			{
-				FileContents += TEXT("EditorStartupMap=/Game/StarterContent/Maps/Minimal_Default") LINE_TERMINATOR;
-				FileContents += TEXT("GameDefaultMap=/Game/StarterContent/Maps/Minimal_Default") LINE_TERMINATOR;
-			}
-
-			if (InProjectInfo.bShouldGenerateCode)
-			{
-				FileContents += FString::Printf(TEXT("GlobalDefaultGameMode=\"/Script/%s.%sGameMode\"") LINE_TERMINATOR, *NewProjectName, *NewProjectName);
-			}
-		}
-
 		if (WriteOutputFile(DefaultEngineIniFilename, FileContents, OutFailReason))
 		{
 			OutCreatedFiles.Add(DefaultEngineIniFilename);
@@ -4294,60 +4210,8 @@ void GameProjectUtils::HarvestCursorSyncLocation( FString& FinalOutput, FString&
 
 bool GameProjectUtils::InsertFeaturePacksIntoINIFile(const FProjectInformation& InProjectInfo, FText& OutFailReason)
 {
-	const FString ProjectName = FPaths::GetBaseFilename(InProjectInfo.ProjectFilename);
-	const FString TemplateName = FPaths::GetBaseFilename(InProjectInfo.TemplateFile);
-	const FString SrcFolder = FPaths::GetPath(InProjectInfo.TemplateFile);
-	const FString DestFolder = FPaths::GetPath(InProjectInfo.ProjectFilename);
-
-	const FString ProjectConfigPath = DestFolder / TEXT("Config");
-	const FString IniFilename = ProjectConfigPath / TEXT("DefaultGame.ini");
-
-	TArray<FString> PackList;
-
-	// First the starter content
-	if (InProjectInfo.bCopyStarterContent)
-	{
-		FString StarterContentName = GameProjectUtils::GetStarterContentName(InProjectInfo);
-		FString StarterPack = FString::Printf(TEXT("InsertPack=(PackSource=\"%s.upack\",PackName=\"StarterContent\")"), *StarterContentName);
-		PackList.Add(StarterPack);
-	}
-
-	if (PackList.Num() != 0)
-	{
-		FString FileOutput;
-		if(FPaths::FileExists(IniFilename) && !FFileHelper::LoadFileToString(FileOutput, *IniFilename))
-		{
-			OutFailReason = LOCTEXT("FailedToReadIni", "Could not read INI file to insert feature packs");
-			return false;
-		}
-
-		FileOutput += LINE_TERMINATOR;
-		FileOutput += TEXT("[StartupActions]");
-		FileOutput += LINE_TERMINATOR;
-		FileOutput += TEXT("bAddPacks=True");
-		FileOutput += LINE_TERMINATOR;
-		for (int32 iLine = 0; iLine < PackList.Num(); ++iLine)
-		{
-			FileOutput += PackList[iLine] + LINE_TERMINATOR;
-		}
-
-		// Register 'StartupActions' as a section to save if one of the default value of its entries
-		// is modified. Otherwise, the mechanism which set 'bAddPacks' to 'False' after the
-		// first load will not be persisted and the selected pack files will be loaded on each launch of the editor.
-		FileOutput += LINE_TERMINATOR;
-		FileOutput += TEXT("[SectionsToSave]");
-		FileOutput += LINE_TERMINATOR;
-		FileOutput += TEXT("+Section=StartupActions");
-		FileOutput += LINE_TERMINATOR;
-
-		if (!FFileHelper::SaveStringToFile(FileOutput, *IniFilename))
-		{
-			OutFailReason = LOCTEXT("FailedToWriteIni", "Could not write INI file to insert feature packs");
-			return false;
-		}
-	}
-
-	return true;
+	OutFailReason = LOCTEXT("DeprecatedStarterPack", "Deprecated: Ability to add Starter Content has been removed");
+	return false;
 }
 
 bool GameProjectUtils::AddSharedContentToProject(const FProjectInformation &InProjectInfo, TArray<FString> &CreatedFiles, FText& OutFailReason)
diff --git a/Engine/Source/Editor/GameProjectGeneration/Private/SProjectDialog.cpp b/Engine/Source/Editor/GameProjectGeneration/Private/SProjectDialog.cpp
--- a/Engine/Source/Editor/GameProjectGeneration/Private/SProjectDialog.cpp
+++ b/Engine/Source/Editor/GameProjectGeneration/Private/SProjectDialog.cpp
@@ -758,11 +758,8 @@ TSharedRef<SWidget> SProjectDialog::MakeTemplateProjectView()
 
 TSharedRef<SWidget> SProjectDialog::MakeHybridView(EProjectDialogModeMode Mode)
 {
-	bCopyStarterContent = GEditor ? GetDefault<UEditorSettings>()->bCopyStarterContentPreference : true;
-
 	SelectedHardwareClassTarget = EHardwareClass::Desktop;
 	SelectedGraphicsPreset = EGraphicsPreset::Maximum;
-	bIsStarterContentAvailable = true;
 
 	// Find all template projects
 	Templates = FindTemplateProjects();
@@ -984,66 +981,6 @@ TSharedRef<SWidget> SProjectDialog::MakeProjectOptionsWidget()
 		];
 	}
 
-	if (!HiddenSettings.Contains(ETemplateSetting::StarterContent))
-	{
-		FProjectInformation ProjectInfo = CreateProjectInfo();
-		bIsStarterContentAvailable = GameProjectUtils::IsStarterContentAvailableForProject(ProjectInfo);
-
-		ProjectOptionsBox->AddSlot()
-			.Padding(0.0f, 8.0f, 0.0f, 0.0f)
-			.AutoHeight()
-			[
-				SNew(SHorizontalBox)
-				.ToolTipText(LOCTEXT("CopyStarterContent_ToolTip", "Enable to include an additional content pack containing simple placeable meshes with basic materials and textures.\nYou can also add the Starter Content to your project later using the Content Browser."))
-				.IsEnabled(this, &SProjectDialog::IsStarterContentAvailable)
-				+ SHorizontalBox::Slot()
-				.Padding(0.0f, 0.0f, 8.0f, 0.0f)
-				.VAlign(VAlign_Center)
-				.HAlign(HAlign_Right)
-				[
-					SNew(STextBlock)
-					.Text(LOCTEXT("IncludeStarterContent", "Starter Content"))
-				]
-				+ SHorizontalBox::Slot()
-				.VAlign(VAlign_Center)
-				[
-					SNew(SCheckBox)
-					.IsChecked(this, &SProjectDialog::GetCopyStarterContentCheckState)
-					.OnCheckStateChanged(this, &SProjectDialog::OnSetCopyStarterContent)
-				]
-			];
-
-		// Only add the option to add starter content if its there to add!
-		if (!bIsStarterContentAvailable)
-		{
-			bCopyStarterContent = false;
-		}
-/*
-
-		TSharedRef<SWidget> Enum = SNew(SOverlay)
-			+ SOverlay::Slot()
-			[
-				SNew(SDecoratedEnumCombo<int32>, MoveTemp(StarterContentOptions))
-				.SelectedEnum()
-				.OnEnumChanged()
-				.Orientation(Orient_Vertical)
-			]
-			+ SOverlay::Slot()
-			.HAlign(HAlign_Right)
-			.VAlign(VAlign_Top)
-			.Padding(4)
-			[
-				// Warning when enabled for mobile, since the current starter content is bad for mobile
-				SNew(SImage)
-				.Image(FAppStyle::GetBrush("Icons.Warning"))
-				.ToolTipText(this, &SNewProjectWizard::GetStarterContentWarningTooltip)
-				.Visibility(this, &SNewProjectWizard::GetStarterContentWarningVisibility)
-			];
-*/
-
-		//AddToProjectSettingsGrid(GridPanel, Enum, Description, CurrentSlot);
-	}
-
 #if 0 // @todo: XR settings cannot be shown at the moment as the setting causes issues with binary builds.
 	if (!HiddenSettings.Contains(ETemplateSetting::XR))
 	{
@@ -1182,11 +1119,6 @@ TSharedRef<SWidget> SProjectDialog::MakeOpenProjectPathArea()
 		];
 }
 
-void SProjectDialog::OnSetCopyStarterContent(ECheckBoxState NewState)
-{
-	bCopyStarterContent = NewState == ECheckBoxState::Checked;
-}
-
 bool SProjectDialog::CanCreateProject() const
 {
 	return bLastGlobalValidityCheckSuccessful && bLastNameAndLocationValidityCheckSuccessful;
@@ -1205,8 +1137,6 @@ void SProjectDialog::OnSetBlueprintOrCppIndex(int32 Index)
 	bShouldGenerateCode = Index == 1;
 	UpdateProjectFileValidity();
 	FProjectInformation ProjectInfo = CreateProjectInfo();
-	bIsStarterContentAvailable = GameProjectUtils::IsStarterContentAvailableForProject(ProjectInfo);
-	bCopyStarterContent &= bIsStarterContentAvailable;
 }
 
 void SProjectDialog::SetHardwareClassTarget(EHardwareClass InHardwareClass)
@@ -1221,14 +1151,6 @@ void SProjectDialog::SetGraphicsPreset(EGraphicsPreset InGraphicsPreset)
 
 void SProjectDialog::HandleTemplateListViewSelectionChanged(TSharedPtr<FTemplateItem> TemplateItem, ESelectInfo::Type SelectInfo)
 {
-	if (TemplateItem.IsValid())
-	{
-		if (TemplateItem->HiddenSettings.Contains(ETemplateSetting::StarterContent))
-		{
-			bCopyStarterContent = false;
-		}
-	}
-
 	(*ProjectOptionsSlot)
 	[
 		MakeProjectOptionsWidget()
@@ -1587,27 +1509,11 @@ FProjectInformation SProjectDialog::CreateProjectInfo() const
 
 	FProjectInformation ProjectInfo;
 	ProjectInfo.bShouldGenerateCode = bShouldGenerateCode;
-	ProjectInfo.bCopyStarterContent = bCopyStarterContent;
 	ProjectInfo.TemplateFile = bShouldGenerateCode ? SelectedTemplate->CodeProjectFile : SelectedTemplate->BlueprintProjectFile;
 	ProjectInfo.TemplateCategory = ActiveCategory;
 	ProjectInfo.bIsEnterpriseProject = SelectedTemplate->bIsEnterprise;
 	ProjectInfo.bIsBlankTemplate = SelectedTemplate->bIsBlankTemplate;
 
-	if (bShouldGenerateCode)
-	{
-		if (SelectedTemplate->CodeTemplateDefs != nullptr)
-		{
-			ProjectInfo.StarterContent = SelectedTemplate->CodeTemplateDefs->StarterContent;
-		}
-	}
-	else
-	{
-		if (SelectedTemplate->BlueprintTemplateDefs != nullptr)
-		{
-			ProjectInfo.StarterContent = SelectedTemplate->BlueprintTemplateDefs->StarterContent;
-		}
-	}
-
 	const TArray<ETemplateSetting>& HiddenSettings = SelectedTemplate->HiddenSettings;
 
 	if (!HiddenSettings.Contains(ETemplateSetting::All))
@@ -1665,7 +1571,6 @@ bool SProjectDialog::CreateProject(const FString& ProjectFile)
 	UEditorSettings* Settings = GetMutableDefault<UEditorSettings>();
 	Settings->CreatedProjectPaths.Remove(CreatedProjectPath);
 	Settings->CreatedProjectPaths.Insert(CreatedProjectPath, 0);
-	Settings->bCopyStarterContentPreference = bCopyStarterContent;
 	Settings->PostEditChange();
 
 	return true;
diff --git a/Engine/Source/Editor/GameProjectGeneration/Private/SProjectDialog.h b/Engine/Source/Editor/GameProjectGeneration/Private/SProjectDialog.h
--- a/Engine/Source/Editor/GameProjectGeneration/Private/SProjectDialog.h
+++ b/Engine/Source/Editor/GameProjectGeneration/Private/SProjectDialog.h
@@ -156,12 +156,6 @@ private:
 
 	TSharedRef<SWidget> MakeOpenProjectPathArea();
 
-	/** Get whether we are copying starter content or not */
-	ECheckBoxState GetCopyStarterContentCheckState() const { return bCopyStarterContent ? ECheckBoxState::Checked : ECheckBoxState::Unchecked; };
-
-	/** Handler for when the "copy starter content" checkbox changes state */
-	void OnSetCopyStarterContent(ECheckBoxState NewState);
-
 	/** Get the images for the selected template preview and category */
 	const FSlateBrush* GetSelectedTemplatePreviewImage() const;
 
@@ -169,8 +163,6 @@ private:
 
 	FReply OnCancel() const;
 	
-	bool IsStarterContentAvailable() const {return bIsStarterContentAvailable; }
-
 	int32 OnGetBlueprintOrCppIndex() const { return bShouldGenerateCode ? 1 : 0; }
 	void OnSetBlueprintOrCppIndex(int32 Index);
 
@@ -244,11 +236,6 @@ private:
 	static TUniquePtr<FSlateBrush> CustomTemplateBrush;
 
 	SVerticalBox::FSlot* ProjectOptionsSlot;
-	/** True if user has selected to copy starter content. */
-	bool bCopyStarterContent;
-
-	/** True if starter content is available for current project. */
-	bool bIsStarterContentAvailable;
 
 	/** Whether or not to enable XR in the created project. */
 	bool bEnableXR;
diff --git a/Engine/Source/Editor/GameProjectGeneration/Private/Tests/GameProjectAutomationTests.cpp b/Engine/Source/Editor/GameProjectGeneration/Private/Tests/GameProjectAutomationTests.cpp
--- a/Engine/Source/Editor/GameProjectGeneration/Private/Tests/GameProjectAutomationTests.cpp
+++ b/Engine/Source/Editor/GameProjectGeneration/Private/Tests/GameProjectAutomationTests.cpp
@@ -73,11 +73,10 @@ namespace GameProjectAutomationUtils
 	 * @param	InTargetedHardware	Target hardware (EHardwareClass)
 	 * @param	InGraphicPreset		Graphics preset (EGraphicsPreset)
 	 * @param	InCategory			Target category (EContentSourceCategory)
-	 * @param	bInCopyStarterContent Should the starter content be copied also
 	 * @param	OutMatchedProjects	Total projects matching criteria
 	 * @param	OutCreatedProjects	Total projects succesfully created
 	 */
-	static void CreateProjectSet(TMap<FName, TArray<TSharedPtr<FTemplateItem>> >& InTemplates, EHardwareClass InTargetedHardware, EGraphicsPreset InGraphicPreset, FName InCategory, bool bInCopyStarterContent, int32 &OutCreatedProjects, int32 &OutMatchedProjects)
+	static void CreateProjectSet(TMap<FName, TArray<TSharedPtr<FTemplateItem>> >& InTemplates, EHardwareClass InTargetedHardware, EGraphicsPreset InGraphicPreset, FName InCategory, int32 &OutCreatedProjects, int32 &OutMatchedProjects)
 	{		
 		// If this is empty, it will use the same name for each project, otherwise it will create a project based on target platform and source template
 		FString TestRootFolder;// = "ProjectTests";
@@ -140,7 +139,6 @@ namespace GameProjectAutomationUtils
 						FProjectInformation ProjectInfo;
 						ProjectInfo.ProjectFilename = DesiredProjectFilename;
 						ProjectInfo.bShouldGenerateCode = ProjectFile == Item->CodeProjectFile;
-						ProjectInfo.bCopyStarterContent = bInCopyStarterContent;
 						ProjectInfo.TemplateFile = ProjectFile;
 						ProjectInfo.TargetedHardware = InTargetedHardware;
 						ProjectInfo.DefaultGraphicsPerformance = InGraphicPreset;
@@ -242,7 +240,6 @@ bool FBuildPromotionNewProjectCreateTest::RunTest(const FString& Parameters)
 	FText FailReason, FailLog;
 	FProjectInformation ProjectInfo;
 	ProjectInfo.ProjectFilename = DesiredProjectFilename;
-	ProjectInfo.bCopyStarterContent = true;
 	if (GameProjectUtils::CreateProject(ProjectInfo, FailReason, FailLog))
 	{
 		UE_LOG(LogGameProjectGenerationTests, Display, TEXT("Generated a new project: %s"), *DesiredProjectFilename);
@@ -331,11 +328,11 @@ bool FCreateBPTemplateProjectAutomationTests::RunTest(const FString& Parameters)
 	TMap<FName, TArray<TSharedPtr<FTemplateItem>> > Templates = SProjectDialog::FindTemplateProjects();
 	int32 OutMatchedProjectsDesk = 0;
 	int32 OutCreatedProjectsDesk = 0;
-	GameProjectAutomationUtils::CreateProjectSet(Templates, EHardwareClass::Desktop, EGraphicsPreset::Maximum, "Game", false, OutMatchedProjectsDesk, OutCreatedProjectsDesk);
+	GameProjectAutomationUtils::CreateProjectSet(Templates, EHardwareClass::Desktop, EGraphicsPreset::Maximum, "Game", OutMatchedProjectsDesk, OutCreatedProjectsDesk);
 
 	int32 OutMatchedProjectsMob = 0;
 	int32 OutCreatedProjectsMob = 0;
-	GameProjectAutomationUtils::CreateProjectSet(Templates, EHardwareClass::Mobile, EGraphicsPreset::Maximum, "Visualization", false, OutMatchedProjectsMob, OutCreatedProjectsMob);
+	GameProjectAutomationUtils::CreateProjectSet(Templates, EHardwareClass::Mobile, EGraphicsPreset::Maximum, "Visualization", OutMatchedProjectsMob, OutCreatedProjectsMob);
 	
 	return ( OutMatchedProjectsDesk == OutCreatedProjectsDesk ) && ( OutMatchedProjectsMob == OutCreatedProjectsMob  );
 }
@@ -357,11 +354,11 @@ bool FCreateCPPTemplateProjectAutomationTests::RunTest(const FString& Parameters
 
 	int32 OutMatchedProjectsDesk = 0;
 	int32 OutCreatedProjectsDesk = 0;
-	GameProjectAutomationUtils::CreateProjectSet(Templates, EHardwareClass::Desktop, EGraphicsPreset::Maximum, "Game", false, OutMatchedProjectsDesk, OutCreatedProjectsDesk);
+	GameProjectAutomationUtils::CreateProjectSet(Templates, EHardwareClass::Desktop, EGraphicsPreset::Maximum, "Game", OutMatchedProjectsDesk, OutCreatedProjectsDesk);
 
 	int32 OutMatchedProjectsMob = 0;
 	int32 OutCreatedProjectsMob = 0;
-	GameProjectAutomationUtils::CreateProjectSet(Templates, EHardwareClass::Mobile, EGraphicsPreset::Maximum, "Game", false, OutMatchedProjectsMob, OutCreatedProjectsMob);
+	GameProjectAutomationUtils::CreateProjectSet(Templates, EHardwareClass::Mobile, EGraphicsPreset::Maximum, "Game", OutMatchedProjectsMob, OutCreatedProjectsMob);
 
 	return (OutMatchedProjectsDesk == OutCreatedProjectsDesk) && (OutMatchedProjectsMob == OutCreatedProjectsMob);
 
diff --git a/Engine/Source/Editor/GameProjectGeneration/Public/GameProjectUtils.h b/Engine/Source/Editor/GameProjectGeneration/Public/GameProjectUtils.h
--- a/Engine/Source/Editor/GameProjectGeneration/Public/GameProjectUtils.h
+++ b/Engine/Source/Editor/GameProjectGeneration/Public/GameProjectUtils.h
@@ -39,7 +39,6 @@ struct FProjectInformation
 	FName TemplateCategory;
 
 	bool bShouldGenerateCode = false;
-	bool bCopyStarterContent = false;
 	bool bIsBlankTemplate = false;
 	bool bIsEnterpriseProject = false;
 
@@ -50,9 +49,6 @@ struct FProjectInformation
 
 	TOptional<EHardwareClass> TargetedHardware;
 	TOptional<EGraphicsPreset> DefaultGraphicsPerformance;
-
-	/** The name of the feature pack to use as starter content. Must be located under FeaturePacks\. */
-	FString StarterContent;
 };
 
 DECLARE_DELEGATE_RetVal_OneParam(bool, FProjectDescriptorModifier, FProjectDescriptor&);
@@ -191,13 +187,16 @@ public:
 	static bool GenerateCodeProjectFiles(const FString& ProjectFilename, FText& OutFailReason, FText& OutFailLog);
 
 	/** Returns true if there are engine starter content files available for instancing into new projects. */
-	static bool IsEngineStarterContentAvailable();
+	UE_DEPRECATED(5.6, "Ability to add Starter Content has been removed")
+	static bool IsEngineStarterContentAvailable() { return false; }
 
 	/** Returns true if the given project is referencing engine starter content pack. */
-	static bool IsUsingEngineStarterContent(const FProjectInformation& ProjectInfo);
+	UE_DEPRECATED(5.6, "Ability to add Starter Content has been removed")
+	static bool IsUsingEngineStarterContent(const FProjectInformation& ProjectInfo) { return false; }
 
 	/** Returns true if there are any starter content packs are available for given project. */
-	static bool IsStarterContentAvailableForProject(const FProjectInformation& ProjectInfo);
+	UE_DEPRECATED(5.6, "Ability to add Starter Content has been removed")
+	static bool IsStarterContentAvailableForProject(const FProjectInformation& ProjectInfo) { return false; }
 
 	/**
 	 * Get the information about any modules referenced in the .uproject file of the currently loaded project
@@ -326,7 +325,8 @@ private:
 	static void AddHardwareConfigValues(const FProjectInformation& InProjectInfo, TArray<FTemplateConfigValue>& ConfigValues);
 	
 	/** Get the name of the starter content pack to use for the given project. */
-	static FString GetStarterContentName(const FProjectInformation& InProjectInfo);
+	UE_DEPRECATED(5.6, "Ability to add Starter Content has been removed")
+	static FString GetStarterContentName(const FProjectInformation& InProjectInfo) { return FString{}; }
 
 	/** Generates a new project without using a template project */
 	static TOptional<FGuid> GenerateProjectFromScratch(const FProjectInformation& InProjectInfo, FText& OutFailReason, FText& OutFailLog);
@@ -338,6 +338,7 @@ private:
 	static bool SetEngineAssociationForForeignProject(const FString& ProjectFileName, FText& OutFailReason);
 
 	/** Insert any required feature packs into the DefaultGame.ini file */
+	UE_DEPRECATED(5.6, "Ability to add Starter Content has been removed")
 	static bool InsertFeaturePacksIntoINIFile(const FProjectInformation& InProjectInfo, FText& OutFailReason);
 
 	/*
diff --git a/Engine/Source/Editor/UnrealEd/Private/Settings/EditorSettings.cpp b/Engine/Source/Editor/UnrealEd/Private/Settings/EditorSettings.cpp
--- a/Engine/Source/Editor/UnrealEd/Private/Settings/EditorSettings.cpp
+++ b/Engine/Source/Editor/UnrealEd/Private/Settings/EditorSettings.cpp
@@ -22,7 +22,6 @@
 UEditorSettings::UEditorSettings(const FObjectInitializer& ObjectInitializer)
 	: Super(ObjectInitializer)
 {
-	bCopyStarterContentPreference = false;
 	AutoScalabilityWorkScaleAmount = 1;
 
 	// Read the current state of the environment variables and cache it.
diff --git a/Engine/Source/Editor/UnrealEd/Public/Settings/EditorSettings.h b/Engine/Source/Editor/UnrealEd/Public/Settings/EditorSettings.h
--- a/Engine/Source/Editor/UnrealEd/Public/Settings/EditorSettings.h
+++ b/Engine/Source/Editor/UnrealEd/Public/Settings/EditorSettings.h
@@ -140,8 +140,11 @@ class UEditorSettings : public UObject
 	UPROPERTY(config)
 	TArray<FString> CreatedProjectPaths;
 
-	UPROPERTY(config)
-	bool bCopyStarterContentPreference;
+#if WITH_EDITORONLY_DATA
+	UE_DEPRECATED(5.6, "Ability to add Starter Content has been removed")
+	UPROPERTY(config, meta = (DeprecatedProperty, DeprecationMessage = "Ability to add Starter Content has been removed"))
+	bool bCopyStarterContentPreference_DEPRECATED = false;
+#endif // WITH_EDITORONLY_DATA
 
 	/** The id's of the surveys completed */
 	UPROPERTY(config)
