[HWRT] Enable r.RayTracing.RayTracingProxies.ProjectEnabled by default for new projects.

diff --git a/Engine/Source/Editor/GameProjectGeneration/Private/GameProjectUtils.cpp b/Engine/Source/Editor/GameProjectGeneration/Private/GameProjectUtils.cpp
--- a/Engine/Source/Editor/GameProjectGeneration/Private/GameProjectUtils.cpp
+++ b/Engine/Source/Editor/GameProjectGeneration/Private/GameProjectUtils.cpp
@@ -182,7 +182,7 @@ namespace
 
 		ConfigValues.Emplace(TEXT("DefaultEngine.ini"),
 			TEXT("/Script/Engine.RendererSettings"),
-			TEXT("r.RayTracing"),
+			TEXT("r.RayTracing.RayTracingProxies.ProjectEnabled"),
 			TEXT("True"),
 			false /* ShouldReplaceExistingValue */);
 	}
