Allow for Variants in TemplateDefs.ini

A template can now define a set of Variants in their TemplateDefs.ini, and users can select zero or one variant from a dropdown in the New Project window.
Each Variant can be defined in one line with their Name, LocalizedDisplayNames and SharedContentPacks. If selected, the SharedContentPacks of the Variant will be copied on project generation.
SharedContentPacks common to all Variants can still be defined like before.

Example:
```
; Feature pack shared by all Variants.
SharedContentPacks=(MountName="Puzzle",DetailLevels=("Standard"))

; Example of a variant
Variants=(Name="PuzzleFirstPerson",LocalizedDisplayNames=((Language="en",Text="First Person Puzzle")),LocalizedDescriptions=((Language="en",Text="A Puzzle with a first person view.")),SharedContentPacks=((DetailLevels=(Standard),MountName="FirstPersonArms"),(DetailLevels=(Standard),MountName="FPWeapon"),(DetailLevels=(Standard,High),MountName="Geometry")))

; Example of a variant
Variants=(Name="PuzzleRacing",LocalizedDisplayNames=((Language="en",Text="Racing Puzzle"), (Language="ko",Text="????????? ??????")),LocalizedDescriptions=((Language="en",Text="A Racing puzzle"), (Language="ko",Text="????????? ??????")),SharedContentPacks=((DetailLevels=(High),MountName="Characters"),(DetailLevels=(Standard,High),MountName="Mannequin"),(DetailLevels=(High),MountName="Track"),(DetailLevels=(Standard),MountName="Vehicles"),(DetailLevels=(Standard),MountName="ConceptCar")))
```

diff --git a/Engine/Source/Editor/GameProjectGeneration/Classes/TemplateProjectDefs.h b/Engine/Source/Editor/GameProjectGeneration/Classes/TemplateProjectDefs.h
--- a/Engine/Source/Editor/GameProjectGeneration/Classes/TemplateProjectDefs.h
+++ b/Engine/Source/Editor/GameProjectGeneration/Classes/TemplateProjectDefs.h
@@ -74,6 +74,31 @@ struct FLocalizedTemplateString
 	static FText GetLocalizedText(const TArray<FLocalizedTemplateString>& LocalizedStrings);
 };
 
+USTRUCT()
+struct FTemplateVariant
+{
+	GENERATED_BODY()
+	
+	UPROPERTY(config)
+	FName Name;
+	
+	UPROPERTY(config)
+	TArray<FLocalizedTemplateString> LocalizedDisplayNames;
+
+	UPROPERTY(config)
+	TArray<FLocalizedTemplateString> LocalizedDescriptions;
+
+	/* Shared feature packs. The files in these packs listed in these structures marked as 'additionalfiles' will be copied on project generation*/
+	UPROPERTY(config)
+	TArray<FFeaturePackLevelSet> SharedContentPacks;
+
+	/** Returns the display name for the current culture, or English if the current culture has no translation */
+	FText GetDisplayNameText() const;
+
+	/** Returns the display name for the current culture, or English if the current culture has no translation */
+	FText GetLocalizedDescription() const;
+};
+
 UENUM()
 enum class ETemplateSetting
 {
@@ -83,6 +108,7 @@ enum class ETemplateSetting
 	StarterContent UE_DEPRECATED(5.6, "Ability to add Starter Content has been removed"),
 	XR,
 	Raytracing_DEPRECATED,
+	Variants,
 	All
 };
 
@@ -163,6 +189,10 @@ public:
 	FString StarterContent_DEPRECATED;
 #endif // WITH_EDITORONLY_DATA
 
+	/** List of variants for this template. Using an Array instead of a Map purely for ease of writing the ini file (allowing one variant per line) */
+	UPROPERTY(config)
+	TArray<FTemplateVariant> Variants;
+
 	/** Fixes up all strings in this definitions object to replace \%TEMPLATENAME\% with the supplied template name and \%PROJECTNAME\% with the supplied project name */
 	void FixupStrings(const FString& TemplateName, const FString& ProjectName);
 
@@ -172,6 +202,9 @@ public:
 	/** Returns the display name for the current culture, or English if the current culture has no translation */
 	FText GetLocalizedDescription();
 
+	/** Finds a Variant with the given name, returns null if not found */
+	const FTemplateVariant* FindVariant(FName VariantName) const;
+
 	/** Does this template generate C++ source? */
 	virtual bool GeneratesCode(const FString& ProjectTemplatePath) const PURE_VIRTUAL(UTemplateProjectDefs::GeneratesCode, return false;)
 
diff --git a/Engine/Source/Editor/GameProjectGeneration/Private/GameProjectUtils.cpp b/Engine/Source/Editor/GameProjectGeneration/Private/GameProjectUtils.cpp
--- a/Engine/Source/Editor/GameProjectGeneration/Private/GameProjectUtils.cpp
+++ b/Engine/Source/Editor/GameProjectGeneration/Private/GameProjectUtils.cpp
@@ -1372,6 +1372,29 @@ UTemplateProjectDefs* GameProjectUtils::LoadTemplateDefs(const FString& ProjectD
 		{
 			UE_LOG(LogGameProjectGeneration, Warning, TEXT("Template '%s' contains 'All' in HiddenSettings in addition to other entries. This is a mistake, and means that all settings will be hidden."), *ProjectDirectory);
 		}
+		
+		if (TemplateDefs->Variants.Num() > 0) // We make sure the Variants are valid
+		{
+			TSet<FName> VariantNames;
+			for (TArray<FTemplateVariant>::TIterator It = TemplateDefs->Variants.CreateIterator(); It; ++It)
+			{
+				const FTemplateVariant& Variant = *It;
+				if (VariantNames.Contains(Variant.Name))
+				{
+					UE_LOG(LogGameProjectGeneration, Warning, TEXT("Template '%s' contains multiple Variants with the name '%s', only the first one will be kept."), *ProjectDirectory, *Variant.Name.ToString());
+				}
+				else if (Variant.Name.IsNone()) // NAME_None is used to represent no Variant selected
+				{
+					UE_LOG(LogGameProjectGeneration, Error, TEXT("Template '%s' contains a Variant with the name '%s' which is not allowed, it will be discarded."), *ProjectDirectory, *Variant.Name.ToString());
+				}
+				else
+				{
+					VariantNames.Add(Variant.Name);
+					continue;
+				}
+				It.RemoveCurrent();
+			}
+		}
 	}
 
 	return TemplateDefs;
@@ -4239,6 +4262,29 @@ bool GameProjectUtils::AddSharedContentToProject(const FProjectInformation &InPr
 			OutFailReason = FText::Format(LOCTEXT("SharedResourceError", "Error adding shared resources for '{TemplateName}'."), Args);
 			return false;
 		}
+
+		if (!TemplateDefs->Variants.IsEmpty() && InProjectInfo.Variant != NAME_None)
+		{
+			const FTemplateVariant* Variant = TemplateDefs->FindVariant(InProjectInfo.Variant);
+			if (!Variant)
+			{
+				FFormatNamedArguments Args;
+				Args.Add(TEXT("TemplateName"), FText::FromString(SrcFolder));
+				OutFailReason = FText::Format(LOCTEXT("SharedResourceError", "Unknown variant specified for '{TemplateName}'."), Args);
+				return false;
+			}
+			
+			TUniquePtr<FFeaturePackContentSource> TempVariantFeaturePack = MakeUnique<FFeaturePackContentSource>();
+			bool bVariantCopied = TempVariantFeaturePack->InsertAdditionalResources(Variant->SharedContentPacks, RequiredDetail, DestFolder, CreatedFiles);
+			if( bVariantCopied == false )
+			{
+				FFormatNamedArguments Args;
+				Args.Add(TEXT("VariantName"), Variant->GetDisplayNameText());
+				Args.Add(TEXT("TemplateName"), FText::FromString(SrcFolder));
+				OutFailReason = FText::Format(LOCTEXT("SharedResourceError", "Error adding variant '{VariantName}' shared resources for '{TemplateName}'."), Args);
+				return false;
+			}
+		}
 	}
 	return true;
 }
diff --git a/Engine/Source/Editor/GameProjectGeneration/Private/SProjectDialog.cpp b/Engine/Source/Editor/GameProjectGeneration/Private/SProjectDialog.cpp
--- a/Engine/Source/Editor/GameProjectGeneration/Private/SProjectDialog.cpp
+++ b/Engine/Source/Editor/GameProjectGeneration/Private/SProjectDialog.cpp
@@ -1,6 +1,7 @@
 // Copyright Epic Games, Inc. All Rights Reserved.
 
 #include "SProjectDialog.h"
+#include "Algo/Transform.h"
 #include "Widgets/Layout/SBorder.h"
 #include "Widgets/SBoxPanel.h"
 #include "GameProjectGenerationModule.h"
@@ -14,6 +15,8 @@
 #include "SourceCodeNavigation.h"
 #include "GameProjectUtils.h"
 #include "Framework/Application/SlateApplication.h"
+#include "Framework/MultiBox/MultiBoxBuilder.h"
+#include "Framework/Commands/UIAction.h"
 #include "Editor.h"
 #include "Internationalization/BreakIterator.h"
 #include "Settings/EditorSettings.h"
@@ -32,6 +35,7 @@
 #include "Widgets/SToolTip.h"
 #include "Widgets/Layout/SGridPanel.h"
 #include "Widgets/SOverlay.h"
+#include "Widgets/Input/SComboBox.h"
 #include "Widgets/Input/SEditableTextBox.h"
 #include "Widgets/Input/SButton.h"
 #include "SWarningOrErrorBox.h"
@@ -44,7 +48,6 @@
 #include "Styling/StyleColors.h"
 #include "Styling/AppStyle.h"
 
-
 #define LOCTEXT_NAMESPACE "GameProjectGeneration"
 
 namespace NewProjectDialogDefs
@@ -838,6 +841,9 @@ TSharedRef<SWidget> SProjectDialog::MakeHybridView(EProjectDialogModeMode Mode)
 TSharedRef<SWidget> SProjectDialog::MakeProjectOptionsWidget()
 {
 	IHardwareTargetingModule& HardwareTargeting = IHardwareTargetingModule::Get();
+	TemplateVariantNames.Empty();
+	SelectedVariantName = NAME_None;
+	const UTemplateProjectDefs* ProjectDefs = GetSelectedTemplateDefs();
 
 	TSharedPtr<SVerticalBox> ProjectOptionsBox;
 
@@ -938,10 +944,9 @@ TSharedRef<SWidget> SProjectDialog::MakeProjectOptionsWidget()
 			]
 			+ SHorizontalBox::Slot()
 			.VAlign(VAlign_Center)
-			.HAlign(HAlign_Left)
+			.HAlign(HAlign_Fill)
 			[
 				SNew(SBox)
-				.WidthOverride(120.0f)
 				[
 					HardwareClassTarget
 				]
@@ -970,16 +975,52 @@ TSharedRef<SWidget> SProjectDialog::MakeProjectOptionsWidget()
 			]
 			+ SHorizontalBox::Slot()
 			.VAlign(VAlign_Center)
-			.HAlign(HAlign_Left)
+			.HAlign(HAlign_Fill)
 			[
 				SNew(SBox)
-				.WidthOverride(120.0f)
 				[
 					GraphicsPreset
 				]
 			]
 		];
 	}
+	
+	if (!HiddenSettings.Contains(ETemplateSetting::Variants) && ProjectDefs)
+	{
+		TemplateVariantNames.Add(NAME_None); // The user can choose to have no variants
+		Algo::Transform(ProjectDefs->Variants, TemplateVariantNames, [](const FTemplateVariant& Variant) { return Variant.Name; });
+		
+		ProjectOptionsBox->AddSlot()
+			.Padding(0.0f, 8.0f, 0.0f, 0.0f)
+			.AutoHeight()
+			[
+				SNew(SHorizontalBox)
+				.ToolTipText(LOCTEXT("Variants_ToolTip", "Which variant would you like to use?"))
+				.Visibility(this, &SProjectDialog::GetVariantsVisibility)
+				+ SHorizontalBox::Slot()
+				.Padding(0.0f, 0.0f, 8.0f, 0.0f)
+				.VAlign(VAlign_Center)
+				.HAlign(HAlign_Right)
+				[
+					SNew(STextBlock)
+					.Text(LOCTEXT("Variant", "Variant"))
+				]
+				+ SHorizontalBox::Slot()
+				.VAlign(VAlign_Center)
+				.HAlign(HAlign_Fill)
+				[
+					SNew(SComboButton)
+					.ContentPadding(FMargin(4.0f, 0.0f))
+					.OnGetMenuContent(this, &SProjectDialog::GetVariantsDropdownContent)
+					.ButtonContent()
+					[
+						SNew(STextBlock)
+						.Text(this, &SProjectDialog::GetVariantsButtonText)
+						.ToolTipText(this, &SProjectDialog::GetVariantsButtonTooltip)
+					]
+				]
+			];
+	}
 
 #if 0 // @todo: XR settings cannot be shown at the moment as the setting causes issues with binary builds.
 	if (!HiddenSettings.Contains(ETemplateSetting::XR))
@@ -1137,6 +1178,20 @@ void SProjectDialog::OnSetBlueprintOrCppIndex(int32 Index)
 	bShouldGenerateCode = Index == 1;
 	UpdateProjectFileValidity();
 	FProjectInformation ProjectInfo = CreateProjectInfo();
+
+	// Recreate the VariantNames as BP and CPP templates can have different variants
+	const UTemplateProjectDefs* ProjectDefs = GetSelectedTemplateDefs();
+	TemplateVariantNames.Empty();
+	if (ProjectDefs)
+	{
+		Algo::Transform(ProjectDefs->Variants, TemplateVariantNames, [](const FTemplateVariant& Variant) { return Variant.Name; });
+	}
+	// Check if this template has a variant of the same name, otherwise resets it
+	const bool bFoundVariant = ProjectDefs && ProjectDefs->FindVariant(SelectedVariantName);
+	if (!bFoundVariant)
+	{
+		SelectedVariantName = NAME_None;
+	}
 }
 
 void SProjectDialog::SetHardwareClassTarget(EHardwareClass InHardwareClass)
@@ -1149,6 +1204,48 @@ void SProjectDialog::SetGraphicsPreset(EGraphicsPreset InGraphicsPreset)
 	SelectedGraphicsPreset = InGraphicsPreset;
 }
 
+EVisibility SProjectDialog::GetVariantsVisibility() const
+{
+	const UTemplateProjectDefs* ProjectDefs = GetSelectedTemplateDefs();
+	return ProjectDefs && !ProjectDefs->Variants.IsEmpty() ? EVisibility::Visible : EVisibility::Collapsed;
+}
+
+TSharedRef<SWidget> SProjectDialog::GetVariantsDropdownContent()
+{
+	FMenuBuilder MenuBuilder(true, nullptr, nullptr, true);
+
+	for (const FName& VariantName : TemplateVariantNames)
+	{
+		const FTemplateVariant* Variant = GetSelectedTemplateDefsVariant(VariantName);
+		MenuBuilder.AddMenuEntry(
+			Variant ? Variant->GetDisplayNameText() : LOCTEXT("NoVariant", "None"),
+			Variant ? Variant->GetLocalizedDescription() : LOCTEXT("NoVariantDescription", "Do not include a variant"),
+			FSlateIcon(),
+			FUIAction
+			(
+				FExecuteAction::CreateLambda([this, VariantName]()
+				{
+					SelectedVariantName = VariantName;
+				})
+			)
+		);
+	}
+
+	return MenuBuilder.MakeWidget();
+}
+
+FText SProjectDialog::GetVariantsButtonText() const
+{
+	const FTemplateVariant* Variant = GetSelectedTemplateDefsVariant();
+    return Variant ? Variant->GetDisplayNameText() : LOCTEXT("NoVariant", "None");
+}
+
+FText SProjectDialog::GetVariantsButtonTooltip() const
+{
+	const FTemplateVariant* Variant = GetSelectedTemplateDefsVariant();
+	return Variant ? Variant->GetLocalizedDescription() : LOCTEXT("NoVariantDescription", "Do not include a variant");
+}
+
 void SProjectDialog::HandleTemplateListViewSelectionChanged(TSharedPtr<FTemplateItem> TemplateItem, ESelectInfo::Type SelectInfo)
 {
 	(*ProjectOptionsSlot)
@@ -1170,6 +1267,21 @@ TSharedPtr<FTemplateItem> SProjectDialog::GetSelectedTemplateItem() const
 	return nullptr;
 }
 
+UTemplateProjectDefs* SProjectDialog::GetSelectedTemplateDefs() const
+{
+	if (const TSharedPtr<FTemplateItem> TemplateItem = GetSelectedTemplateItem())
+	{
+		return bShouldGenerateCode ? TemplateItem->CodeTemplateDefs : TemplateItem->BlueprintTemplateDefs;
+	}
+	return nullptr;
+}
+
+const FTemplateVariant* SProjectDialog::GetSelectedTemplateDefsVariant(FName VariantName) const
+{
+	const UTemplateProjectDefs* ProjectDefs = GetSelectedTemplateDefs();
+	return ProjectDefs ? ProjectDefs->FindVariant(VariantName) : nullptr;
+}
+
 namespace
 {
 	FString MakeSortKey(const FString& TemplateKey)
@@ -1510,6 +1622,7 @@ FProjectInformation SProjectDialog::CreateProjectInfo() const
 	FProjectInformation ProjectInfo;
 	ProjectInfo.bShouldGenerateCode = bShouldGenerateCode;
 	ProjectInfo.TemplateFile = bShouldGenerateCode ? SelectedTemplate->CodeProjectFile : SelectedTemplate->BlueprintProjectFile;
+	ProjectInfo.Variant = SelectedVariantName;
 	ProjectInfo.TemplateCategory = ActiveCategory;
 	ProjectInfo.bIsEnterpriseProject = SelectedTemplate->bIsEnterprise;
 	ProjectInfo.bIsBlankTemplate = SelectedTemplate->bIsBlankTemplate;
diff --git a/Engine/Source/Editor/GameProjectGeneration/Private/SProjectDialog.h b/Engine/Source/Editor/GameProjectGeneration/Private/SProjectDialog.h
--- a/Engine/Source/Editor/GameProjectGeneration/Private/SProjectDialog.h
+++ b/Engine/Source/Editor/GameProjectGeneration/Private/SProjectDialog.h
@@ -67,6 +67,16 @@ class SProjectDialog : public SCompoundWidget
 private:
 	/** Accessor for the currently selected template item */
 	TSharedPtr<FTemplateItem> GetSelectedTemplateItem() const;
+	
+	/** Accessor for the TemplateDefs of the currently selected template item. Will return either the BP or CPP version based on bShouldGenerateCode */
+	UTemplateProjectDefs* GetSelectedTemplateDefs() const;
+	/** Accessor for the Variant matching the SelectedVariantName for the currently selected TemplateDefs (see GetSelectedTemplateDefs) */
+	const FTemplateVariant* GetSelectedTemplateDefsVariant() const
+	{
+		return GetSelectedTemplateDefsVariant(SelectedVariantName);
+	}
+	/** Accessor for the Variant with the given name for the currently selected TemplateDefs (see GetSelectedTemplateDefs). Returns null if not found */
+	const FTemplateVariant* GetSelectedTemplateDefsVariant(FName VariantName) const;
 
 	/** Helper function to allow direct lookup of the selected item's properties on a delegate */
 	/** TRet should be defaulted but VS2012 doesn't allow default template arguments on non-class templates */
@@ -172,6 +182,11 @@ private:
 	void SetGraphicsPreset(EGraphicsPreset InGraphicsPreset);
 	EGraphicsPreset GetGraphicsPreset() const { return SelectedGraphicsPreset; }
 
+	EVisibility GetVariantsVisibility() const;
+	TSharedRef< SWidget > GetVariantsDropdownContent();
+	FText GetVariantsButtonText() const;
+	FText GetVariantsButtonTooltip() const;
+
 	/** Handler for when the selection changes in the template list */
 	void HandleTemplateListViewSelectionChanged(TSharedPtr<FTemplateItem> TemplateItem, ESelectInfo::Type SelectInfo);
 
@@ -231,6 +246,9 @@ private:
 
 	EGraphicsPreset SelectedGraphicsPreset;
 
+	FName SelectedVariantName;
+	TArray<FName> TemplateVariantNames;
+
 	TUniquePtr<FSlateBrush> RecentProjectsBrush;
 
 	static TUniquePtr<FSlateBrush> CustomTemplateBrush;
diff --git a/Engine/Source/Editor/GameProjectGeneration/Private/TemplateProjectDefs.cpp b/Engine/Source/Editor/GameProjectGeneration/Private/TemplateProjectDefs.cpp
--- a/Engine/Source/Editor/GameProjectGeneration/Private/TemplateProjectDefs.cpp
+++ b/Engine/Source/Editor/GameProjectGeneration/Private/TemplateProjectDefs.cpp
@@ -3,6 +3,7 @@
 
 #include "TemplateProjectDefs.h"
 
+#include "Algo/Find.h"
 #include "HAL/Platform.h"
 #include "Internationalization/Culture.h"
 #include "Internationalization/Internationalization.h"
@@ -51,6 +52,16 @@ FText FLocalizedTemplateString::GetLocalizedText(const TArray<FLocalizedTemplate
 	return FText();
 }
 
+FText FTemplateVariant::GetDisplayNameText() const
+{
+	return FLocalizedTemplateString::GetLocalizedText(LocalizedDisplayNames);
+}
+
+FText FTemplateVariant::GetLocalizedDescription() const
+{
+	return FLocalizedTemplateString::GetLocalizedText(LocalizedDescriptions);
+}
+
 FTemplateConfigValue::FTemplateConfigValue(const FString& InFile, const FString& InSection, const FString& InKey, const FString& InValue, bool InShouldReplaceExistingValue)
 	: ConfigFile(InFile)
 	, ConfigSection(InSection)
@@ -113,6 +124,14 @@ FText UTemplateProjectDefs::GetLocalizedDescription()
 	return FLocalizedTemplateString::GetLocalizedText(LocalizedDescriptions);
 }
 
+const FTemplateVariant* UTemplateProjectDefs::FindVariant(FName VariantName) const
+{
+	return VariantName == NAME_None ? nullptr : Algo::FindByPredicate(Variants, [&](const FTemplateVariant& Variant)
+	{
+		return Variant.Name == VariantName;
+	});
+}
+
 void UTemplateProjectDefs::FixString(FString& InOutStringToFix, const FString& TemplateName, const FString& ProjectName)
 {
 	InOutStringToFix.ReplaceInline(TEXT("%TEMPLATENAME%"), *TemplateName, ESearchCase::CaseSensitive);
diff --git a/Engine/Source/Editor/GameProjectGeneration/Public/GameProjectUtils.h b/Engine/Source/Editor/GameProjectGeneration/Public/GameProjectUtils.h
--- a/Engine/Source/Editor/GameProjectGeneration/Public/GameProjectUtils.h
+++ b/Engine/Source/Editor/GameProjectGeneration/Public/GameProjectUtils.h
@@ -37,6 +37,7 @@ struct FProjectInformation
 	FString ProjectFilename;
 	FString TemplateFile;
 	FName TemplateCategory;
+	FName Variant;
 
 	bool bShouldGenerateCode = false;
 	bool bIsBlankTemplate = false;
