Restoring SM6 default setting.

diff --git a/Engine/Source/Editor/GameProjectGeneration/Private/GameProjectUtils.cpp b/Engine/Source/Editor/GameProjectGeneration/Private/GameProjectUtils.cpp
--- a/Engine/Source/Editor/GameProjectGeneration/Private/GameProjectUtils.cpp
+++ b/Engine/Source/Editor/GameProjectGeneration/Private/GameProjectUtils.cpp
@@ -2104,6 +2104,31 @@ void GameProjectUtils::AddHardwareConfigValues(const FProjectInformation& InProj
 			}
 		}
 	}
+
+	// Don't override these settings for templates
+	if (InProjectInfo.TemplateFile.IsEmpty())
+	{
+		// New projects always have DX12 by default on Windows
+		ConfigValues.Emplace(TEXT("DefaultEngine.ini"),
+			TEXT("/Script/WindowsTargetPlatform.WindowsTargetSettings"),
+			TEXT("DefaultGraphicsRHI"),
+			TEXT("DefaultGraphicsRHI_DX12"),
+			false /* ShouldReplaceExistingValue */);
+
+		// Force clear D3D12TargetedShaderFormats since the BaseEngine list can change at any time.
+		ConfigValues.Emplace(TEXT("DefaultEngine.ini"),
+			TEXT("/Script/WindowsTargetPlatform.WindowsTargetSettings"),
+			TEXT("!D3D12TargetedShaderFormats"),
+			TEXT("ClearArray"),
+			false /* ShouldReplaceExistingValue */);
+
+		// New projects always have DX12 only supporting SM6 by default on Windows
+		ConfigValues.Emplace(TEXT("DefaultEngine.ini"),
+			TEXT("/Script/WindowsTargetPlatform.WindowsTargetSettings"),
+			TEXT("+D3D12TargetedShaderFormats"),
+			TEXT("PCD3D_SM6"),
+			false /* ShouldReplaceExistingValue */);
+	}
 }
 
 bool GameProjectUtils::GenerateConfigFiles(const FProjectInformation& InProjectInfo, TArray<FString>& OutCreatedFiles, FText& OutFailReason, FGuid& OutProjectID)
