Fix for Variants dropdown menu not getting repopulated with options.

The `UTemplateProjectDefs` from which the Variants are found is currently retrieved too early. It might not match the one expected as `bShouldGenerateCode` is set a few lines below.
This CL creates a new function to properly and consistently handle the retrieval of the Variants, and ensure it is called at the right time.

diff --git a/Engine/Source/Editor/GameProjectGeneration/Private/SProjectDialog.cpp b/Engine/Source/Editor/GameProjectGeneration/Private/SProjectDialog.cpp
--- a/Engine/Source/Editor/GameProjectGeneration/Private/SProjectDialog.cpp
+++ b/Engine/Source/Editor/GameProjectGeneration/Private/SProjectDialog.cpp
@@ -841,9 +841,6 @@ TSharedRef<SWidget> SProjectDialog::MakeHybridView(EProjectDialogModeMode Mode)
 TSharedRef<SWidget> SProjectDialog::MakeProjectOptionsWidget()
 {
 	IHardwareTargetingModule& HardwareTargeting = IHardwareTargetingModule::Get();
-	TemplateVariantNames.Empty();
-	SelectedVariantName = NAME_None;
-	const UTemplateProjectDefs* ProjectDefs = GetSelectedTemplateDefs();
 
 	TSharedPtr<SVerticalBox> ProjectOptionsBox;
 
@@ -985,11 +982,9 @@ TSharedRef<SWidget> SProjectDialog::MakeProjectOptionsWidget()
 		];
 	}
 	
-	if (!HiddenSettings.Contains(ETemplateSetting::Variants) && ProjectDefs)
+	ResetVariants(); // Needs to be called after we have set bShouldGenerateCode
+	if (!HiddenSettings.Contains(ETemplateSetting::Variants))
 	{
-		TemplateVariantNames.Add(NAME_None); // The user can choose to have no variants
-		Algo::Transform(ProjectDefs->Variants, TemplateVariantNames, [](const FTemplateVariant& Variant) { return Variant.Name; });
-		
 		ProjectOptionsBox->AddSlot()
 			.Padding(0.0f, 8.0f, 0.0f, 0.0f)
 			.AutoHeight()
@@ -1179,19 +1174,7 @@ void SProjectDialog::OnSetBlueprintOrCppIndex(int32 Index)
 	UpdateProjectFileValidity();
 	FProjectInformation ProjectInfo = CreateProjectInfo();
 
-	// Recreate the VariantNames as BP and CPP templates can have different variants
-	const UTemplateProjectDefs* ProjectDefs = GetSelectedTemplateDefs();
-	TemplateVariantNames.Empty();
-	if (ProjectDefs)
-	{
-		Algo::Transform(ProjectDefs->Variants, TemplateVariantNames, [](const FTemplateVariant& Variant) { return Variant.Name; });
-	}
-	// Check if this template has a variant of the same name, otherwise resets it
-	const bool bFoundVariant = ProjectDefs && ProjectDefs->FindVariant(SelectedVariantName);
-	if (!bFoundVariant)
-	{
-		SelectedVariantName = NAME_None;
-	}
+	ResetVariants();
 }
 
 void SProjectDialog::SetHardwareClassTarget(EHardwareClass InHardwareClass)
@@ -1246,6 +1229,25 @@ FText SProjectDialog::GetVariantsButtonTooltip() const
 	return Variant ? Variant->GetLocalizedDescription() : LOCTEXT("NoVariantDescription", "Do not include a variant");
 }
 
+void SProjectDialog::ResetVariants()
+{
+	// Recreate the VariantNames. BP and CPP templates can have different variants
+	TemplateVariantNames.Empty();
+	TemplateVariantNames.Add(NAME_None); // The user can choose to have no variants
+	
+	const UTemplateProjectDefs* ProjectDefs = GetSelectedTemplateDefs();
+	if (ProjectDefs)
+	{
+		Algo::Transform(ProjectDefs->Variants, TemplateVariantNames, [](const FTemplateVariant& Variant) { return Variant.Name; });
+	}
+	// Check if this template has a variant of the same name as the one previously selected, otherwise resets it
+	const bool bFoundVariant = ProjectDefs && ProjectDefs->FindVariant(SelectedVariantName);
+	if (!bFoundVariant)
+	{
+		SelectedVariantName = NAME_None;
+	}
+}
+
 void SProjectDialog::HandleTemplateListViewSelectionChanged(TSharedPtr<FTemplateItem> TemplateItem, ESelectInfo::Type SelectInfo)
 {
 	(*ProjectOptionsSlot)
diff --git a/Engine/Source/Editor/GameProjectGeneration/Private/SProjectDialog.h b/Engine/Source/Editor/GameProjectGeneration/Private/SProjectDialog.h
--- a/Engine/Source/Editor/GameProjectGeneration/Private/SProjectDialog.h
+++ b/Engine/Source/Editor/GameProjectGeneration/Private/SProjectDialog.h
@@ -182,6 +182,7 @@ private:
 	void SetGraphicsPreset(EGraphicsPreset InGraphicsPreset);
 	EGraphicsPreset GetGraphicsPreset() const { return SelectedGraphicsPreset; }
 
+	void ResetVariants();
 	EVisibility GetVariantsVisibility() const;
 	TSharedRef< SWidget > GetVariantsDropdownContent();
 	FText GetVariantsButtonText() const;
