Fix Variants LOCTEXT

diff --git a/Engine/Source/Editor/GameProjectGeneration/Private/GameProjectUtils.cpp b/Engine/Source/Editor/GameProjectGeneration/Private/GameProjectUtils.cpp
--- a/Engine/Source/Editor/GameProjectGeneration/Private/GameProjectUtils.cpp
+++ b/Engine/Source/Editor/GameProjectGeneration/Private/GameProjectUtils.cpp
@@ -4295,7 +4295,7 @@ bool GameProjectUtils::AddSharedContentToProject(const FProjectInformation &InPr
 			{
 				FFormatNamedArguments Args;
 				Args.Add(TEXT("TemplateName"), FText::FromString(SrcFolder));
-				OutFailReason = FText::Format(LOCTEXT("SharedResourceError", "Unknown variant specified for '{TemplateName}'."), Args);
+				OutFailReason = FText::Format(LOCTEXT("UnknownVariantError", "Unknown variant specified for '{TemplateName}'."), Args);
 				return false;
 			}
 			
@@ -4306,7 +4306,7 @@ bool GameProjectUtils::AddSharedContentToProject(const FProjectInformation &InPr
 				FFormatNamedArguments Args;
 				Args.Add(TEXT("VariantName"), Variant->GetDisplayNameText());
 				Args.Add(TEXT("TemplateName"), FText::FromString(SrcFolder));
-				OutFailReason = FText::Format(LOCTEXT("SharedResourceError", "Error adding variant '{VariantName}' shared resources for '{TemplateName}'."), Args);
+				OutFailReason = FText::Format(LOCTEXT("VariantSharedResourceError", "Error adding variant '{VariantName}' shared resources for '{TemplateName}'."), Args);
 				return false;
 			}
 		}
