Fixed the consequences of UI scaling when adding blueprint classes via the following New Class Dialog improvements:
- Window is now resizable.
- Component widgets are now contained within a SScrollbox.
- Text explanations are now set to wrap.
- Path picker now fills.

[RN]

diff --git a/Engine/Source/Editor/GameProjectGeneration/Private/GameProjectUtils.cpp b/Engine/Source/Editor/GameProjectGeneration/Private/GameProjectUtils.cpp
--- a/Engine/Source/Editor/GameProjectGeneration/Private/GameProjectUtils.cpp
+++ b/Engine/Source/Editor/GameProjectGeneration/Private/GameProjectUtils.cpp
@@ -1117,7 +1117,7 @@ void GameProjectUtils::OpenAddToProjectDialog(const FAddToProjectConfig& Config,
 		SNew(SWindow)
 		.Title( WindowTitle )
 		.ClientSize( WindowSize )
-		.SizingRule( ESizingRule::FixedSize )
+		.SizingRule( ESizingRule::UserSized )
 		.SupportsMinimize(false) .SupportsMaximize(false);
 
 	TSharedRef<SNewClassDialog> NewClassDialog =
diff --git a/Engine/Source/Editor/GameProjectGeneration/Private/SNewClassDialog.cpp b/Engine/Source/Editor/GameProjectGeneration/Private/SNewClassDialog.cpp
--- a/Engine/Source/Editor/GameProjectGeneration/Private/SNewClassDialog.cpp
+++ b/Engine/Source/Editor/GameProjectGeneration/Private/SNewClassDialog.cpp
@@ -300,6 +300,7 @@ void SNewClassDialog::Construct( const FArguments& InArgs )
 					)
 				.OnCanceled(this, &SNewClassDialog::CancelClicked)
 				.OnFinished(this, &SNewClassDialog::FinishClicked)
+				.ShowHistoryButtons(ClassDomain == EClassDomain::Native || !ParentClassInfo.IsSet())
 				.InitialPageIndex(ParentClassInfo.IsSet() ? 1 : 0)
 				.PageFooter()
 				[
@@ -489,11 +490,12 @@ void SNewClassDialog::Construct( const FArguments& InArgs )
 				+SWizard::Page()
 				.OnEnter(this, &SNewClassDialog::OnNamePageEntered)
 				[
-					SNew(SVerticalBox)
+					SNew(SScrollBox)
+					.AllowContentToShrink(true)
 
 					// Title
-					+SVerticalBox::Slot()
-					.AutoHeight()
+					+SScrollBox::Slot()
+					.AutoSize()
 					.Padding(0.0f)
 					[
 						SNew(STextBlock)
@@ -502,8 +504,9 @@ void SNewClassDialog::Construct( const FArguments& InArgs )
 						.TransformPolicy(ETextTransformPolicy::ToUpper)
 					]
 
-					+SVerticalBox::Slot()
-					.FillHeight(1.f)
+					+SScrollBox::Slot()
+					.FillSize(1.0f)
+					.MinSize(358.0f)
 					.Padding(0.0f, 10.0f)
 					[
 						SNew(SVerticalBox)
@@ -514,6 +517,7 @@ void SNewClassDialog::Construct( const FArguments& InArgs )
 						[
 							SNew(STextBlock)
 							.Text(LOCTEXT("ClassNameDescription", "Enter a name for your new class. Class names may only contain alphanumeric characters, and may not contain a space.") )
+								.AutoWrapText(true)
 						]
 
 						+SVerticalBox::Slot()
@@ -525,6 +529,7 @@ void SNewClassDialog::Construct( const FArguments& InArgs )
 								LOCTEXT("ClassNameDetails_Native", "When you click the \"Create\" button below, a header (.h) file and a source (.cpp) file will be made using this name.") :
 								FText::Format(LOCTEXT("ClassNameDetails_Blueprint", "When you click the \"Create\" button below, a new class inheriting from {0} will be created."), ParentClassInfo.IsSet() ? ParentClassInfo.GetClassName() : FText::FromStringView(TEXT("Blueprint")))
 								)
+								.AutoWrapText(true)
 						]
 
 						// Name Error label
@@ -540,223 +545,217 @@ void SNewClassDialog::Construct( const FArguments& InArgs )
 
 						// Properties
 						+SVerticalBox::Slot()
-						.AutoHeight()
+						.FillHeight(1.0f)
 						[
 							SNew(SBorder)
 							.BorderImage(FAppStyle::GetBrush("DetailsView.CategoryTop"))
 							.BorderBackgroundColor(FLinearColor(0.6f, 0.6f, 0.6f, 1.0f ))
 							.Padding(FMargin(6.0f, 4.0f, 7.0f, 4.0f))
 							[
-								SNew(SVerticalBox)
+								SNew(SGridPanel)
+								.FillColumn(1, 1.0f)
+								.FillRow(2, 1.0f)
+
+								// Class type label
+								+ SGridPanel::Slot(0, 0)
+								.VAlign(VAlign_Center)
+								.Padding(0.0f, 0.0f, 12.0f, 0.0f)
+								[
+									SNew(STextBlock)
+									.Visibility(ClassDomain == EClassDomain::Blueprint ? EVisibility::Collapsed : EVisibility::Visible)
+									.Text(LOCTEXT("ClassTypeLabel", "Class Type"))
+								]
 
-								+SVerticalBox::Slot()
-								.AutoHeight()
-								.Padding(0.0f)
+								+SGridPanel::Slot(1,0)
+								.VAlign(VAlign_Center)
+								.HAlign(HAlign_Left)
+								.Padding(2.0f)
+								[
+									SNew(SSegmentedControl<GameProjectUtils::EClassLocation>)
+									.Visibility(ClassDomain == EClassDomain::Blueprint ? EVisibility::Collapsed : EVisibility::Visible)
+									.OnValueChanged(this, &SNewClassDialog::OnClassLocationChanged)
+									.Value(this, &SNewClassDialog::IsClassLocationActive)
+									+ SSegmentedControl<GameProjectUtils::EClassLocation>::Slot(GameProjectUtils::EClassLocation::Public)
+									.Text(LOCTEXT("Public", "Public"))
+									.ToolTip(LOCTEXT("ClassLocation_Public", "A public class can be included and used inside other modules in addition to the module it resides in"))
+									+ SSegmentedControl<GameProjectUtils::EClassLocation>::Slot(GameProjectUtils::EClassLocation::Private)
+									.Text(LOCTEXT("Private", "Private"))
+									.ToolTip(LOCTEXT("ClassLocation_Private", "A private class can only be included and used within the module it resides in"))
+								]
+								// Name label
+								+SGridPanel::Slot(0, 1)
+								.VAlign(VAlign_Center)
+								.Padding(0.0f, 0.0f, 12.0f, 0.0f)
 								[
-									SNew(SGridPanel)
-									.FillColumn(1, 1.0f)
-									// Class type label
+									SNew(STextBlock)
+									.Text( LOCTEXT( "NameLabel", "Name" ) )
+								]
 
-									+ SGridPanel::Slot(0, 0)
-									.VAlign(VAlign_Center)
-									.Padding(0.0f, 0.0f, 12.0f, 0.0f)
+								// Name edit box
+								+SGridPanel::Slot(1, 1)
+								.Padding(0.0f, 3.0f)
+								.VAlign(VAlign_Center)
+								[
+									SNew(SBox)
+									.HeightOverride(EditableTextHeight)
+									.AddMetaData<FTutorialMetaData>(TEXT("ClassName"))
 									[
-										SNew(STextBlock)
-										.Visibility(ClassDomain == EClassDomain::Blueprint ? EVisibility::Collapsed : EVisibility::Visible)
-										.Text(LOCTEXT("ClassTypeLabel", "Class Type"))
-									]
+										SNew(SHorizontalBox)
 
-									+SGridPanel::Slot(1,0)
-									.VAlign(VAlign_Center)
-									.HAlign(HAlign_Left)
-									.Padding(2.0f)
-									[
-										SNew(SSegmentedControl<GameProjectUtils::EClassLocation>)
-										.Visibility(ClassDomain == EClassDomain::Blueprint ? EVisibility::Collapsed : EVisibility::Visible)
-										.OnValueChanged(this, &SNewClassDialog::OnClassLocationChanged)
-										.Value(this, &SNewClassDialog::IsClassLocationActive)
-										+ SSegmentedControl<GameProjectUtils::EClassLocation>::Slot(GameProjectUtils::EClassLocation::Public)
-										.Text(LOCTEXT("Public", "Public"))
-										.ToolTip(LOCTEXT("ClassLocation_Public", "A public class can be included and used inside other modules in addition to the module it resides in"))
-										+ SSegmentedControl<GameProjectUtils::EClassLocation>::Slot(GameProjectUtils::EClassLocation::Private)
-										.Text(LOCTEXT("Private", "Private"))
-										.ToolTip(LOCTEXT("ClassLocation_Private", "A private class can only be included and used within the module it resides in"))
-									]
-									// Name label
-									+SGridPanel::Slot(0, 1)
-									.VAlign(VAlign_Center)
-									.Padding(0.0f, 0.0f, 12.0f, 0.0f)
-									[
-										SNew(STextBlock)
-										.Text( LOCTEXT( "NameLabel", "Name" ) )
+										+SHorizontalBox::Slot()
+										.FillWidth(.7f)
+										[
+											SAssignNew( ClassNameEditBox, SEditableTextBox)
+											.Text( this, &SNewClassDialog::OnGetClassNameText )
+											.OnTextChanged( this, &SNewClassDialog::OnClassNameTextChanged )
+											.OnTextCommitted( this, &SNewClassDialog::OnClassNameTextCommitted )
+										]
+
+										+SHorizontalBox::Slot()
+										.AutoWidth()
+										.Padding(6.0f, 0.0f, 0.0f, 0.0f)
+										[
+											SAssignNew(AvailableModulesCombo, SComboBox<TSharedPtr<FModuleContextInfo>>)
+											.Visibility(ClassDomain == EClassDomain::Blueprint ? EVisibility::Collapsed : EVisibility::Visible)
+											.ToolTipText( LOCTEXT("ModuleComboToolTip", "Choose the target module for your new class") )
+											.OptionsSource( &AvailableModules )
+											.InitiallySelectedItem( SelectedModuleInfo )
+											.OnSelectionChanged( this, &SNewClassDialog::SelectedModuleComboBoxSelectionChanged )
+											.OnGenerateWidget( this, &SNewClassDialog::MakeWidgetForSelectedModuleCombo )
+											[
+												SNew(STextBlock)
+												.Text( this, &SNewClassDialog::GetSelectedModuleComboText )
+											]
+										]
 									]
+								]
 
-									// Name edit box
-									+SGridPanel::Slot(1, 1)
-									.Padding(0.0f, 3.0f)
-									.VAlign(VAlign_Center)
+								// Path label
+								+SGridPanel::Slot(0, 2)
+								.VAlign(ClassDomain == EClassDomain::Blueprint ? VAlign_Top : VAlign_Center)
+								.Padding(0.0f, 0.0f, 12.0f, 0.0f)
+								[
+									SNew(STextBlock)
+									.Text( LOCTEXT( "PathLabel", "Path" ) )
+								]
+
+								// Path edit box
+								+SGridPanel::Slot(1, 2)
+								.Padding(0.0f, 3.0f)
+								.VAlign(VAlign_Fill)
+								[
+									SNew(SVerticalBox)
+
+									// Blueprint Class asset path
+									+ SVerticalBox::Slot()
+									.FillHeight(1.0f)
+									.Padding(0.0f)
 									[
 										SNew(SBox)
-										.HeightOverride(EditableTextHeight)
-										.AddMetaData<FTutorialMetaData>(TEXT("ClassName"))
+										.Visibility(ClassDomain == EClassDomain::Blueprint ? EVisibility::Visible : EVisibility::Collapsed)
 										[
-											SNew(SHorizontalBox)
-
-											+SHorizontalBox::Slot()
-											.FillWidth(.7f)
+											SNew(SVerticalBox)
+											
+											+SVerticalBox::Slot()
+											.AutoHeight()
 											[
-												SAssignNew( ClassNameEditBox, SEditableTextBox)
-												.Text( this, &SNewClassDialog::OnGetClassNameText )
-												.OnTextChanged( this, &SNewClassDialog::OnClassNameTextChanged )
-												.OnTextCommitted( this, &SNewClassDialog::OnClassNameTextCommitted )
+												SNew(STextBlock)
+												.Text(this, &SNewClassDialog::OnGetClassPathText)
 											]
 
-											+SHorizontalBox::Slot()
-											.AutoWidth()
-											.Padding(6.0f, 0.0f, 0.0f, 0.0f)
+											+SVerticalBox::Slot()
+											.FillHeight(1.0f)
 											[
-												SAssignNew(AvailableModulesCombo, SComboBox<TSharedPtr<FModuleContextInfo>>)
-												.Visibility(ClassDomain == EClassDomain::Blueprint ? EVisibility::Collapsed : EVisibility::Visible)
-												.ToolTipText( LOCTEXT("ModuleComboToolTip", "Choose the target module for your new class") )
-												.OptionsSource( &AvailableModules )
-												.InitiallySelectedItem( SelectedModuleInfo )
-												.OnSelectionChanged( this, &SNewClassDialog::SelectedModuleComboBoxSelectionChanged )
-												.OnGenerateWidget( this, &SNewClassDialog::MakeWidgetForSelectedModuleCombo )
-												[
-													SNew(STextBlock)
-													.Text( this, &SNewClassDialog::GetSelectedModuleComboText )
-												]
+												ContentBrowser.CreatePathPicker(BlueprintPathConfig)
 											]
 										]
 									]
 
-									// Path label
-									+SGridPanel::Slot(0, 2)
-									.VAlign(ClassDomain == EClassDomain::Blueprint ? VAlign_Top : VAlign_Center)
-									.Padding(0.0f, 0.0f, 12.0f, 0.0f)
-									[
-										SNew(STextBlock)
-										.Text( LOCTEXT( "PathLabel", "Path" ) )
-									]
-
-									// Path edit box
-									+SGridPanel::Slot(1, 2)
-									.Padding(0.0f, 3.0f)
-									.VAlign(VAlign_Center)
+									// Native C++ path
+									+ SVerticalBox::Slot()
+									.Padding(0.0f)
+									.AutoHeight()
 									[
-										SNew(SVerticalBox)
-
-										// Blueprint Class asset path
-										+ SVerticalBox::Slot()
-										.Padding(0.0f)
+										SNew(SBox)
+										.Visibility(ClassDomain == EClassDomain::Blueprint ? EVisibility::Collapsed : EVisibility::Visible)
+										.HeightOverride(EditableTextHeight)
+										.AddMetaData<FTutorialMetaData>(TEXT("Path"))
 										[
-											SNew(SBox)
-											// Height override to force the visibility of a scrollbar (our parent is autoheight)
-											.HeightOverride(220.0f)
-											.Visibility(ClassDomain == EClassDomain::Blueprint ? EVisibility::Visible : EVisibility::Collapsed)
-											[
-												SNew(SVerticalBox)
-												
-												+SVerticalBox::Slot()
-												.AutoHeight()
-												[
-													SNew(STextBlock)
-													.Text(this, &SNewClassDialog::OnGetClassPathText)
-												]
+											SNew(SHorizontalBox)
 
-												+SVerticalBox::Slot()
-												[
-													ContentBrowser.CreatePathPicker(BlueprintPathConfig)
-												]
+											+SHorizontalBox::Slot()
+											.FillWidth(1.0f)
+											[
+												SNew(SEditableTextBox)
+												.Text(this, &SNewClassDialog::OnGetClassPathText)
+												.OnTextChanged(this, &SNewClassDialog::OnClassPathTextChanged)
 											]
-										]
 
-										// Native C++ path
-										+ SVerticalBox::Slot()
-										.Padding(0.0f)
-										.AutoHeight()
-										[
-											SNew(SBox)
-											.Visibility(ClassDomain == EClassDomain::Blueprint ? EVisibility::Collapsed : EVisibility::Visible)
-											.HeightOverride(EditableTextHeight)
-											.AddMetaData<FTutorialMetaData>(TEXT("Path"))
+											+SHorizontalBox::Slot()
+											.AutoWidth()
+											.Padding(6.0f, 1.0f, 0.0f, 0.0f)
 											[
-												SNew(SHorizontalBox)
-
-												+SHorizontalBox::Slot()
-												.FillWidth(1.0f)
+												SNew(SButton)
+												.VAlign(VAlign_Center)
+												.ButtonStyle(FAppStyle::Get(), "SimpleButton")
+												.OnClicked(this, &SNewClassDialog::HandleChooseFolderButtonClicked)
 												[
-													SNew(SEditableTextBox)
-													.Text(this, &SNewClassDialog::OnGetClassPathText)
-													.OnTextChanged(this, &SNewClassDialog::OnClassPathTextChanged)
-												]
-
-												+SHorizontalBox::Slot()
-												.AutoWidth()
-												.Padding(6.0f, 1.0f, 0.0f, 0.0f)
-												[
-													SNew(SButton)
-													.VAlign(VAlign_Center)
-													.ButtonStyle(FAppStyle::Get(), "SimpleButton")
-													.OnClicked(this, &SNewClassDialog::HandleChooseFolderButtonClicked)
-													[
-														SNew(SImage)
-														.Image(FAppStyle::Get().GetBrush("Icons.FolderClosed"))
-														.ColorAndOpacity(FSlateColor::UseForeground())
-													]
+													SNew(SImage)
+													.Image(FAppStyle::Get().GetBrush("Icons.FolderClosed"))
+													.ColorAndOpacity(FSlateColor::UseForeground())
 												]
 											]
 										]
 									]
+								]
+
+								// Header output label
+								+SGridPanel::Slot(0, 3)
+								.VAlign(VAlign_Center)
+								.Padding(0.0f, 0.0f, 12.0f, 0.0f)
+								[
+									SNew(STextBlock)
+									.Visibility(ClassDomain == EClassDomain::Blueprint ? EVisibility::Collapsed : EVisibility::Visible)
+									.Text( LOCTEXT( "HeaderFileLabel", "Header File" ) )
+								]
 
-									// Header output label
-									+SGridPanel::Slot(0, 3)
+								// Header output text
+								+SGridPanel::Slot(1, 3)
+								.Padding(0.0f, 3.0f)
+								.VAlign(VAlign_Center)
+								[
+									SNew(SBox)
+									.Visibility(ClassDomain == EClassDomain::Blueprint ? EVisibility::Collapsed : EVisibility::Visible)
 									.VAlign(VAlign_Center)
-									.Padding(0.0f, 0.0f, 12.0f, 0.0f)
+									.HeightOverride(EditableTextHeight)
 									[
 										SNew(STextBlock)
-										.Visibility(ClassDomain == EClassDomain::Blueprint ? EVisibility::Collapsed : EVisibility::Visible)
-										.Text( LOCTEXT( "HeaderFileLabel", "Header File" ) )
+										.Text(this, &SNewClassDialog::OnGetClassHeaderFileText)
 									]
+								]
 
-									// Header output text
-									+SGridPanel::Slot(1, 3)
-									.Padding(0.0f, 3.0f)
-									.VAlign(VAlign_Center)
-									[
-										SNew(SBox)
-										.Visibility(ClassDomain == EClassDomain::Blueprint ? EVisibility::Collapsed : EVisibility::Visible)
-										.VAlign(VAlign_Center)
-										.HeightOverride(EditableTextHeight)
-										[
-											SNew(STextBlock)
-											.Text(this, &SNewClassDialog::OnGetClassHeaderFileText)
-										]
-									]
+								// Source output label
+								+SGridPanel::Slot(0, 4)
+								.VAlign(VAlign_Center)
+								.Padding(0.0f, 0.0f, 12.0f, 0.0f)
+								[
+									SNew(STextBlock)
+									.Visibility(ClassDomain == EClassDomain::Blueprint ? EVisibility::Collapsed : EVisibility::Visible)
+									.Text( LOCTEXT( "SourceFileLabel", "Source File" ) )
+								]
 
-									// Source output label
-									+SGridPanel::Slot(0, 4)
+								// Source output text
+								+SGridPanel::Slot(1, 4)
+								.Padding(0.0f, 3.0f)
+								.VAlign(VAlign_Center)
+								[
+									SNew(SBox)
+									.Visibility(ClassDomain == EClassDomain::Blueprint ? EVisibility::Collapsed : EVisibility::Visible)
 									.VAlign(VAlign_Center)
-									.Padding(0.0f, 0.0f, 12.0f, 0.0f)
+									.HeightOverride(EditableTextHeight)
 									[
 										SNew(STextBlock)
-										.Visibility(ClassDomain == EClassDomain::Blueprint ? EVisibility::Collapsed : EVisibility::Visible)
-										.Text( LOCTEXT( "SourceFileLabel", "Source File" ) )
-									]
-
-									// Source output text
-									+SGridPanel::Slot(1, 4)
-									.Padding(0.0f, 3.0f)
-									.VAlign(VAlign_Center)
-									[
-										SNew(SBox)
-										.Visibility(ClassDomain == EClassDomain::Blueprint ? EVisibility::Collapsed : EVisibility::Visible)
-										.VAlign(VAlign_Center)
-										.HeightOverride(EditableTextHeight)
-										[
-											SNew(STextBlock)
-											.Text(this, &SNewClassDialog::OnGetClassSourceFileText)
-										]
+										.Text(this, &SNewClassDialog::OnGetClassSourceFileText)
 									]
 								]
 							]
diff --git a/Engine/Source/Runtime/AppFramework/Private/Widgets/Workflow/SWizard.cpp b/Engine/Source/Runtime/AppFramework/Private/Widgets/Workflow/SWizard.cpp
--- a/Engine/Source/Runtime/AppFramework/Private/Widgets/Workflow/SWizard.cpp
+++ b/Engine/Source/Runtime/AppFramework/Private/Widgets/Workflow/SWizard.cpp
@@ -32,6 +32,7 @@ void SWizard::Construct( const FArguments& InArgs )
 	OnCanceled = InArgs._OnCanceled;
 	OnFinished = InArgs._OnFinished;
 	OnFirstPageBackClicked = InArgs._OnFirstPageBackClicked;
+	bShowHistoryButtons = InArgs._ShowHistoryButtons;
 
 	TSharedPtr<SVerticalBox> PageListBox;
 	TSharedPtr< SUniformGridPanel > ButtonGrid; 
@@ -74,6 +75,7 @@ void SWizard::Construct( const FArguments& InArgs )
 				]
 
 				+ SHorizontalBox::Slot()
+				.VAlign(VAlign_Fill)
 				.FillWidth(1.0f)
 				[
 					// widget switcher
@@ -356,12 +358,22 @@ FReply SWizard::HandleNextButtonClicked()
 
 bool SWizard::HandleNextButtonIsEnabled() const
 {
+	if (!bShowHistoryButtons)
+	{
+		return false;
+	}
+
 	int32 NextPage = GetNextPageIndex();
 	return CanShowPage(NextPage);
 }
 
 EVisibility SWizard::HandleNextButtonVisibility() const
 {
+	if (!bShowHistoryButtons)
+	{
+		return EVisibility::Collapsed;
+	}
+
 	int32 NextPage = GetNextPageIndex();
 	if (Pages.IsValidIndex(NextPage))
 	{
@@ -400,19 +412,17 @@ FReply SWizard::HandlePrevButtonClicked()
 	{
 		return OnFirstPageBackClicked.Execute();
 	}
-	else
-	{
-		int32 PrevPage = GetPrevPageIndex();
-		ShowPage(PrevPage);
-		BreadcrumbTrail->PopCrumb();
-	}
+
+	int32 PrevPage = GetPrevPageIndex();
+	ShowPage(PrevPage);
+	BreadcrumbTrail->PopCrumb();
 
 	return FReply::Handled();
 }
 
 bool SWizard::HandlePrevButtonIsEnabled() const
 {
-	if ( WidgetSwitcher->GetActiveWidgetIndex() == 0 && OnFirstPageBackClicked.IsBound() )
+	if (bShowHistoryButtons && WidgetSwitcher->GetActiveWidgetIndex() == 0 && OnFirstPageBackClicked.IsBound())
 	{
 		return true;
 	}
@@ -423,6 +433,11 @@ bool SWizard::HandlePrevButtonIsEnabled() const
 
 EVisibility SWizard::HandlePrevButtonVisibility() const
 {
+	if (!bShowHistoryButtons)
+	{
+		return EVisibility::Collapsed;
+	}
+
 	int32 PrevPage = GetPrevPageIndex();
 	if (PrevPage != INDEX_NONE || OnFirstPageBackClicked.IsBound())
 	{
diff --git a/Engine/Source/Runtime/AppFramework/Public/Widgets/Workflow/SWizard.h b/Engine/Source/Runtime/AppFramework/Public/Widgets/Workflow/SWizard.h
--- a/Engine/Source/Runtime/AppFramework/Public/Widgets/Workflow/SWizard.h
+++ b/Engine/Source/Runtime/AppFramework/Public/Widgets/Workflow/SWizard.h
@@ -180,6 +180,7 @@ public:
 		, _ShowBreadcrumbs(false)
 		, _BreadcrumbButtonStyle(&FCoreStyle::Get().GetWidgetStyle< FButtonStyle >("BreadcrumbButton"))
 		, _BreadcrumbTextStyle(&FCoreStyle::Get().GetWidgetStyle< FTextBlockStyle >("NormalText"))
+		, _ShowHistoryButtons(true)
 		, _ShowPageTitle(false)
 		, _PageTitleTextStyle(&FCoreStyle::Get().GetWidgetStyle< FTextBlockStyle >("Wizard.PageTitle"))
 		, _PageFooter()
@@ -253,6 +254,9 @@ public:
 		/** The text style to use for breadcrumbs if ShowBreadcrumbs is set to true. */
 		SLATE_STYLE_ARGUMENT(FTextBlockStyle, BreadcrumbTextStyle);
 
+		/** Holds a flag indicating whether the previous and next buttons can be shown (default = true) */
+		SLATE_ARGUMENT(bool, ShowHistoryButtons)
+
 		/** Holds a flag indicating whether a page title should be shown for each page (default = false) */
 		SLATE_ARGUMENT(bool, ShowPageTitle)
 
@@ -393,6 +397,9 @@ private:
 	// Holds the breadcrumb trail.
 	TSharedPtr<SBreadcrumbTrail<int32>> BreadcrumbTrail;
 
+	// Whether to show the "Previous" and "Next" buttons
+	bool bShowHistoryButtons;
+
 private:
 
 	// Holds a delegate to be invoked when the 'Cancel' button has been clicked.
