Updated headers using UnrealCodeFixup to make sure dllstorage is on methods/staticvars instead of types. Part 5/n

diff --git a/Engine/Source/Editor/GameProjectGeneration/Public/GameProjectUtils.h b/Engine/Source/Editor/GameProjectGeneration/Public/GameProjectUtils.h
--- a/Engine/Source/Editor/GameProjectGeneration/Public/GameProjectUtils.h
+++ b/Engine/Source/Editor/GameProjectGeneration/Public/GameProjectUtils.h
@@ -16,6 +16,8 @@
 #include "Templates/SharedPointer.h"
 #include "UObject/NameTypes.h"
 
+#define UE_API GAMEPROJECTGENERATION_API
+
 class FText;
 class SNotificationItem;
 class UClass;
@@ -54,7 +56,7 @@ struct FProjectInformation
 
 DECLARE_DELEGATE_RetVal_OneParam(bool, FProjectDescriptorModifier, FProjectDescriptor&);
 
-class GAMEPROJECTGENERATION_API GameProjectUtils
+class GameProjectUtils
 {
 public:
 	/** Where is this class located within the Source folder? */
@@ -113,28 +115,28 @@ public:
 	};
 
 	/** Returns true if the project filename is properly formed and does not conflict with another project */
-	static bool IsValidProjectFileForCreation(const FString& ProjectFile, FText& OutFailReason);
+	static UE_API bool IsValidProjectFileForCreation(const FString& ProjectFile, FText& OutFailReason);
 
 	/** Opens the specified project, if it exists. Returns true if the project file is valid. On failure, OutFailReason will be populated. */
-	static bool OpenProject(const FString& ProjectFile, FText& OutFailReason);
+	static UE_API bool OpenProject(const FString& ProjectFile, FText& OutFailReason);
 
 	/** Opens the code editing IDE for the specified project, if it exists. Returns true if the IDE could be opened. On failure, OutFailReason will be populated. */
-	static bool OpenCodeIDE(const FString& ProjectFile, FText& OutFailReason);
+	static UE_API bool OpenCodeIDE(const FString& ProjectFile, FText& OutFailReason);
 
 	/** Creates the specified project file and all required folders. If TemplateFile is non-empty, it will be used as the template for creation. On failure, OutFailReason will be populated. */
-	static bool CreateProject(const FProjectInformation& InProjectInfo, FText& OutFailReason, FText& OutFailLog, TArray<FString>* OutCreatedFiles = nullptr);
+	static UE_API bool CreateProject(const FProjectInformation& InProjectInfo, FText& OutFailReason, FText& OutFailLog, TArray<FString>* OutCreatedFiles = nullptr);
 
 	/** Prompts the user to update their project file, if necessary. */
-	static void CheckForOutOfDateGameProjectFile();
+	static UE_API void CheckForOutOfDateGameProjectFile();
 
 	/** Warn the user if the project filename is invalid in case they renamed it outside the editor */
-	static void CheckAndWarnProjectFilenameValid();
+	static UE_API void CheckAndWarnProjectFilenameValid();
 
 	/** Checks out the current project file (or prompts to make writable) */
-	static void TryMakeProjectFileWriteable(const FString& ProjectFile);
+	static UE_API void TryMakeProjectFileWriteable(const FString& ProjectFile);
 
 	/** Updates the given project file to an engine identifier. Returns true if the project was updated successfully or if no update was needed */
-	static bool UpdateGameProject(const FString& ProjectFile, const FString& EngineIdentifier, FText& OutFailReason);
+	static UE_API bool UpdateGameProject(const FString& ProjectFile, const FString& EngineIdentifier, FText& OutFailReason);
 
 	/**
 	 * Opens a dialog to add code files or blueprints to the current project.
@@ -142,50 +144,50 @@ public:
 	 * @param	Config			Configuration options for the dialog
 	 * @param	InDomain		The domain of the class we're creating (native or blueprint)
 	 */
-	static void OpenAddToProjectDialog(const FAddToProjectConfig& Config, EClassDomain InDomain);
+	static UE_API void OpenAddToProjectDialog(const FAddToProjectConfig& Config, EClassDomain InDomain);
 
 	/** Returns true if the specified class name is properly formed and does not conflict with another class */
-	static bool IsValidClassNameForCreation(const FString& NewClassName, FText& OutFailReason);
+	static UE_API bool IsValidClassNameForCreation(const FString& NewClassName, FText& OutFailReason);
 
 	/** Returns true if the specified class name is properly formed and does not conflict with another class, including source/header files */
-	static bool IsValidClassNameForCreation(const FString& NewClassName, const FModuleContextInfo& ModuleInfo, const TSet<FString>& DisallowedHeaderNames, FText& OutFailReason);
+	static UE_API bool IsValidClassNameForCreation(const FString& NewClassName, const FModuleContextInfo& ModuleInfo, const TSet<FString>& DisallowedHeaderNames, FText& OutFailReason);
 
 	/** Returns true if the specified class is a valid base class for the given module */
-	static bool IsValidBaseClassForCreation(const UClass* InClass, const FModuleContextInfo& InModuleInfo);
+	static UE_API bool IsValidBaseClassForCreation(const UClass* InClass, const FModuleContextInfo& InModuleInfo);
 
 	/** Returns true if the specified class is a valid base class for any of the given modules */
-	static bool IsValidBaseClassForCreation(const UClass* InClass, const TArray<FModuleContextInfo>& InModuleInfoArray);
+	static UE_API bool IsValidBaseClassForCreation(const UClass* InClass, const TArray<FModuleContextInfo>& InModuleInfoArray);
 
 	/** Adds new source code to the project. When returning Succeeded or FailedToHotReload, OutSyncFileAndLineNumber will be the the preferred target file to sync in the users code editing IDE, formatted for use with GenericApplication::GotoLineInSource */
-	static EAddCodeToProjectResult AddCodeToProject(const FString& NewClassName, const FString& NewClassPath, const FModuleContextInfo& ModuleInfo, const FNewClassInfo ParentClassInfo, const TSet<FString>& DisallowedHeaderNames, FString& OutHeaderFilePath, FString& OutCppFilePath, FText& OutFailReason);
+	static UE_API EAddCodeToProjectResult AddCodeToProject(const FString& NewClassName, const FString& NewClassPath, const FModuleContextInfo& ModuleInfo, const FNewClassInfo ParentClassInfo, const TSet<FString>& DisallowedHeaderNames, FString& OutHeaderFilePath, FString& OutCppFilePath, FText& OutFailReason);
 
 	/** Adds new source code to the project. When returning Succeeded or FailedToHotReload, OutSyncFileAndLineNumber will be the the preferred target file to sync in the users code editing IDE, formatted for use with GenericApplication::GotoLineInSource */
-	static EAddCodeToProjectResult AddCodeToProject(const FString& NewClassName, const FString& NewClassPath, const FModuleContextInfo& ModuleInfo, const FNewClassInfo ParentClassInfo, const TSet<FString>& DisallowedHeaderNames, FString& OutHeaderFilePath, FString& OutCppFilePath, FText& OutFailReason, EReloadStatus& OutReloadStatus);
+	static UE_API EAddCodeToProjectResult AddCodeToProject(const FString& NewClassName, const FString& NewClassPath, const FModuleContextInfo& ModuleInfo, const FNewClassInfo ParentClassInfo, const TSet<FString>& DisallowedHeaderNames, FString& OutHeaderFilePath, FString& OutCppFilePath, FText& OutFailReason, EReloadStatus& OutReloadStatus);
 
 	/** Loads a list of template categories defined in the TemplateCategories.ini file in the specified folder */
-	static UTemplateCategories* LoadTemplateCategories(const FString& RootDir);
+	static UE_API UTemplateCategories* LoadTemplateCategories(const FString& RootDir);
 
 	/** Loads a template project definitions object from the TemplateDefs.ini file in the specified project */
-	static UTemplateProjectDefs* LoadTemplateDefs(const FString& ProjectDirectory);
+	static UE_API UTemplateProjectDefs* LoadTemplateDefs(const FString& ProjectDirectory);
 
 	/** @return The number of code files in the currently loaded project */
-	static int32 GetProjectCodeFileCount();
+	static UE_API int32 GetProjectCodeFileCount();
 
 	/**
 	* Retrieves file and size info about the project's source directory
 	* @param OutNumFiles Contains the number of files within the source directory
 	* @param OutDirectorySize Contains the combined size of all files in the directory
 	*/
-	static void GetProjectSourceDirectoryInfo(int32& OutNumFiles, int64& OutDirectorySize);
+	static UE_API void GetProjectSourceDirectoryInfo(int32& OutNumFiles, int64& OutDirectorySize);
 
 	/** Returns the uproject template filename for the default project template. */
-	static FString GetDefaultProjectTemplateFilename();
+	static UE_API FString GetDefaultProjectTemplateFilename();
 
 	/** Compiles a project while showing a progress bar, and offers to open the IDE if it fails. */
-	static bool BuildCodeProject(const FString& ProjectFilename);
+	static UE_API bool BuildCodeProject(const FString& ProjectFilename);
 
 	/** Creates code project files for a new game project. On failure, OutFailReason and OutFailLog will be populated. */
-	static bool GenerateCodeProjectFiles(const FString& ProjectFilename, FText& OutFailReason, FText& OutFailLog);
+	static UE_API bool GenerateCodeProjectFiles(const FString& ProjectFilename, FText& OutFailReason, FText& OutFailLog);
 
 	/** Returns true if there are engine starter content files available for instancing into new projects. */
 	UE_DEPRECATED(5.6, "Ability to add Starter Content has been removed")
@@ -202,17 +204,17 @@ public:
 	/**
 	 * Get the information about any modules referenced in the .uproject file of the currently loaded project
 	 */
-	static const TArray<FModuleContextInfo>& GetCurrentProjectModules();
+	static UE_API const TArray<FModuleContextInfo>& GetCurrentProjectModules();
 
 	/**
 	* Get the information about any modules in any of the plugins in the currently loaded project (Ignores Engine Plugins)
 	*/
-	static TArray<FModuleContextInfo> GetCurrentProjectPluginModules();
+	static UE_API TArray<FModuleContextInfo> GetCurrentProjectPluginModules();
 
 	/**
 	 * Reset the cached result of information about any modules referenced in the .uproject file of the currently loaded project
 	 */
-	static void ResetCurrentProjectModulesCache();
+	static UE_API void ResetCurrentProjectModulesCache();
 
 	/**
 	 * Check to see if the given path is a valid place to put source code for this project (exists within the source root path)
@@ -223,7 +225,7 @@ public:
 	 *
 	 * @return	true if the path is valid, false otherwise
 	 */
-	static bool IsValidSourcePath(const FString& InPath, const FModuleContextInfo& ModuleInfo, FText* const OutFailReason = nullptr);
+	static UE_API bool IsValidSourcePath(const FString& InPath, const FModuleContextInfo& ModuleInfo, FText* const OutFailReason = nullptr);
 
 	/**
 	 * Given the path provided, work out where generated .h and .cpp files would be placed
@@ -236,7 +238,7 @@ public:
 	 *
 	 * @return	false if the paths are invalid
 	 */
-	static bool CalculateSourcePaths(const FString& InPath, const FModuleContextInfo& ModuleInfo, FString& OutHeaderPath, FString& OutSourcePath, FText* const OutFailReason = nullptr);
+	static UE_API bool CalculateSourcePaths(const FString& InPath, const FModuleContextInfo& ModuleInfo, FString& OutHeaderPath, FString& OutSourcePath, FText* const OutFailReason = nullptr);
 
 	/**
 	 * Given the path provided, work out where it's located within the Source folder
@@ -248,10 +250,10 @@ public:
 	 *
 	 * @return	false if the paths are invalid
 	 */
-	static bool GetClassLocation(const FString& InPath, const FModuleContextInfo& ModuleInfo, EClassLocation& OutClassLocation, FText* const OutFailReason = nullptr);
+	static UE_API bool GetClassLocation(const FString& InPath, const FModuleContextInfo& ModuleInfo, EClassLocation& OutClassLocation, FText* const OutFailReason = nullptr);
 
 	/** Creates a copy of a project directory in order to upgrade it. */
-	static EProjectDuplicateResult DuplicateProjectForUpgrade( const FString& InProjectFile, FString& OutNewProjectFile );
+	static UE_API EProjectDuplicateResult DuplicateProjectForUpgrade( const FString& InProjectFile, FString& OutNewProjectFile );
 
 	/**
 	 * Update the list of supported target platforms based upon the parameters provided
@@ -260,52 +262,52 @@ public:
 	 * @param	InPlatformName		Name of the platform to target (eg, WindowsClient)
 	 * @param	bIsSupported		true if the platform should be supported by this project, false if it should not
 	 */
-	static void UpdateSupportedTargetPlatforms(const FName& InPlatformName, const bool bIsSupported);
+	static UE_API void UpdateSupportedTargetPlatforms(const FName& InPlatformName, const bool bIsSupported);
 
 	/** Clear the list of supported target platforms */
-	static void ClearSupportedTargetPlatforms();
+	static UE_API void ClearSupportedTargetPlatforms();
 
 	/** Returns the path to the module's include header */
-	static FString DetermineModuleIncludePath(const FModuleContextInfo& ModuleInfo, const FString& FileRelativeTo);
+	static UE_API FString DetermineModuleIncludePath(const FModuleContextInfo& ModuleInfo, const FString& FileRelativeTo);
 
 	/** Creates the basic source code for a new project. On failure, OutFailReason will be populated. */
-	static bool GenerateBasicSourceCode(TArray<FString>& OutCreatedFiles, FText& OutFailReason);
+	static UE_API bool GenerateBasicSourceCode(TArray<FString>& OutCreatedFiles, FText& OutFailReason);
 
 	/** Generates a Build.cs file for a game module */
-	static bool GenerateGameModuleBuildFile(const FString& NewBuildFileName, const FString& ModuleName, const TArray<FString>& PublicDependencyModuleNames, const TArray<FString>& PrivateDependencyModuleNames, FText& OutFailReason);
+	static UE_API bool GenerateGameModuleBuildFile(const FString& NewBuildFileName, const FString& ModuleName, const TArray<FString>& PublicDependencyModuleNames, const TArray<FString>& PrivateDependencyModuleNames, FText& OutFailReason);
 
 	/** Generates a Build.cs file for a plugin module. Set 'bUseExplicitOrSharedPCHs' to false to disable IWYU conventions. */
-	static bool GeneratePluginModuleBuildFile(const FString& NewBuildFileName, const FString& ModuleName, const TArray<FString>& PublicDependencyModuleNames, const TArray<FString>& PrivateDependencyModuleNames, FText& OutFailReason, bool bUseExplicitOrSharedPCHs = true);
+	static UE_API bool GeneratePluginModuleBuildFile(const FString& NewBuildFileName, const FString& ModuleName, const TArray<FString>& PublicDependencyModuleNames, const TArray<FString>& PrivateDependencyModuleNames, FText& OutFailReason, bool bUseExplicitOrSharedPCHs = true);
 
 	/** Generates a module .cpp file, intended for plugin use */
-	static bool GeneratePluginModuleCPPFile(const FString& CPPFileName, const FString& ModuleName, const FString& StartupSourceCode, FText& OutFailReason);
+	static UE_API bool GeneratePluginModuleCPPFile(const FString& CPPFileName, const FString& ModuleName, const FString& StartupSourceCode, FText& OutFailReason);
 
 	/** Generates a module .h file, intended for plugin use */
-	static bool GeneratePluginModuleHeaderFile(const FString& HeaderFileName, const TArray<FString>& PublicHeaderIncludes, FText& OutFailReason);
+	static UE_API bool GeneratePluginModuleHeaderFile(const FString& HeaderFileName, const TArray<FString>& PublicHeaderIncludes, FText& OutFailReason);
 
 	/** Returns true if the currently loaded project has code files */
-	static bool ProjectHasCodeFiles();
+	static UE_API bool ProjectHasCodeFiles();
 
 	/** Returns the contents of the specified template file */
-	static bool ReadTemplateFile(const FString& TemplateFileName, FString& OutFileContents, FText& OutFailReason);
+	static UE_API bool ReadTemplateFile(const FString& TemplateFileName, FString& OutFileContents, FText& OutFailReason);
 
 	/** Writes an output file. OutputFilename includes a path */
-	static bool WriteOutputFile(const FString& OutputFilename, const FString& OutputFileContents, FText& OutFailReason);
+	static UE_API bool WriteOutputFile(const FString& OutputFilename, const FString& OutputFileContents, FText& OutFailReason);
 
 	/** Returns a comma delimited string comprised of all the elements in InList. If bPlaceQuotesAroundEveryElement, every element is within quotes. */
-	static FString MakeCommaDelimitedList(const TArray<FString>& InList, bool bPlaceQuotesAroundEveryElement = true);
+	static UE_API FString MakeCommaDelimitedList(const TArray<FString>& InList, bool bPlaceQuotesAroundEveryElement = true);
 
 	/** Checks the name for illegal characters */
-	static bool NameContainsOnlyLegalCharacters(const FString& TestName, FString& OutIllegalCharacters);
+	static UE_API bool NameContainsOnlyLegalCharacters(const FString& TestName, FString& OutIllegalCharacters);
 
 	/** Checks the name to see if it matches platform module names */
-	static bool NameMatchesPlatformModuleName(const FString& TestName);
+	static UE_API bool NameMatchesPlatformModuleName(const FString& TestName);
 
 	/** Returns a list of #include lines formed from InList */
-	static FString MakeIncludeList(const TArray<FString>& InList);
+	static UE_API FString MakeIncludeList(const TArray<FString>& InList);
 
 	/** Deletes the specified list of files that were created during file creation */
-	static void DeleteCreatedFiles(const FString& RootFolder, const TArray<FString>& CreatedFiles);
+	static UE_API void DeleteCreatedFiles(const FString& RootFolder, const TArray<FString>& CreatedFiles);
 
 	/**
 	 * Update the list of plugin directories to scan
@@ -315,32 +317,32 @@ public:
 	 * @param	bAddOrRemove true if the directory should be added to this project, false if it should not
 	 * @return Whether the plugin directory list was changed
 	 */
-	static bool UpdateAdditionalPluginDirectory(const FString& InDir, const bool bAddOrRemove);
+	static UE_API bool UpdateAdditionalPluginDirectory(const FString& InDir, const bool bAddOrRemove);
 
 	/** Gets the default build settings version for UBT */
-	static const TCHAR* GetDefaultBuildSettingsVersion();
+	static UE_API const TCHAR* GetDefaultBuildSettingsVersion();
 
 private:
 	
 	/** Add hardware-specific config values such as the target platform and RHI. */
-	static void AddHardwareConfigValues(const FProjectInformation& InProjectInfo, TArray<FTemplateConfigValue>& ConfigValues);
+	static UE_API void AddHardwareConfigValues(const FProjectInformation& InProjectInfo, TArray<FTemplateConfigValue>& ConfigValues);
 	
 	/** Get the name of the starter content pack to use for the given project. */
 	UE_DEPRECATED(5.6, "Ability to add Starter Content has been removed")
 	static FString GetStarterContentName(const FProjectInformation& InProjectInfo) { return FString{}; }
 
 	/** Generates a new project without using a template project */
-	static TOptional<FGuid> GenerateProjectFromScratch(const FProjectInformation& InProjectInfo, FText& OutFailReason, FText& OutFailLog);
+	static UE_API TOptional<FGuid> GenerateProjectFromScratch(const FProjectInformation& InProjectInfo, FText& OutFailReason, FText& OutFailLog);
 
 	/** Generates a new project using a template project */
-	static TOptional<FGuid> CreateProjectFromTemplate(const FProjectInformation& InProjectInfo, FText& OutFailReason, FText& OutFailLog, TArray<FString>* OutCreatedFiles = nullptr);
+	static UE_API TOptional<FGuid> CreateProjectFromTemplate(const FProjectInformation& InProjectInfo, FText& OutFailReason, FText& OutFailLog, TArray<FString>* OutCreatedFiles = nullptr);
 
 	/** Sets the engine association for a new project. Handles foreign and non-foreign projects. */
-	static bool SetEngineAssociationForForeignProject(const FString& ProjectFileName, FText& OutFailReason);
+	static UE_API bool SetEngineAssociationForForeignProject(const FString& ProjectFileName, FText& OutFailReason);
 
 	/** Insert any required feature packs into the DefaultGame.ini file */
 	UE_DEPRECATED(5.6, "Ability to add Starter Content has been removed")
-	static bool InsertFeaturePacksIntoINIFile(const FProjectInformation& InProjectInfo, FText& OutFailReason);
+	static UE_API bool InsertFeaturePacksIntoINIFile(const FProjectInformation& InProjectInfo, FText& OutFailReason);
 
 	/*
 	 * Insert the addition files from any feature packs specified in the temapalte defs file
@@ -350,67 +352,67 @@ private:
 	 *
 	 * @returns true if no errors
 	 */
-	static bool AddSharedContentToProject(const FProjectInformation &InProjectInfo, TArray<FString> &CreatedFiles, FText& OutFailReason);
+	static UE_API bool AddSharedContentToProject(const FProjectInformation &InProjectInfo, TArray<FString> &CreatedFiles, FText& OutFailReason);
 
 	/** Returns the template defs ini filename */
-	static FString GetTemplateDefsFilename();
+	static UE_API FString GetTemplateDefsFilename();
 
 	/** Returns the include header path for a given fully specified, normalized file path */
-	static FString GetIncludePathForFile(const FString& InFullFilePath, const FString& ModuleRootPath);
+	static UE_API FString GetIncludePathForFile(const FString& InFullFilePath, const FString& ModuleRootPath);
 
 	/** Returns true if the project file exists on disk */
-	static bool ProjectFileExists(const FString& ProjectFile);
+	static UE_API bool ProjectFileExists(const FString& ProjectFile);
 
 	/** Returns true if any project files exist in the given folder */
-	static bool AnyProjectFilesExistInFolder(const FString& Path);
+	static UE_API bool AnyProjectFilesExistInFolder(const FString& Path);
 
 	/** Returns true if file cleanup on failure is enabled, false if not */
-	static bool CleanupIsEnabled();
+	static UE_API bool CleanupIsEnabled();
 
 	/** Creates ini files for a new project. On failure, OutFailReason will be populated. */
-	static bool GenerateConfigFiles(const FProjectInformation& InProjectInfo, TArray<FString>& OutCreatedFiles, FText& OutFailReason, FGuid& OutProjectID);
+	static UE_API bool GenerateConfigFiles(const FProjectInformation& InProjectInfo, TArray<FString>& OutCreatedFiles, FText& OutFailReason, FGuid& OutProjectID);
 
 	/* Creates new ini files for a specific project's platform configurations. */
-	static bool GeneratePlatformConfigFiles(const FProjectInformation& InProjectInfo, FText& OutFailReason);
+	static UE_API bool GeneratePlatformConfigFiles(const FProjectInformation& InProjectInfo, FText& OutFailReason);
 
 	/** Creates the basic source code for a new project. On failure, OutFailReason will be populated. */
-	static bool GenerateBasicSourceCode(const FString& NewProjectSourcePath, const FString& NewProjectName, const FString& NewProjectRoot, TArray<FString>& OutGeneratedStartupModuleNames, TArray<FString>& OutCreatedFiles, FText& OutFailReason);
+	static UE_API bool GenerateBasicSourceCode(const FString& NewProjectSourcePath, const FString& NewProjectName, const FString& NewProjectRoot, TArray<FString>& OutGeneratedStartupModuleNames, TArray<FString>& OutCreatedFiles, FText& OutFailReason);
 
 	/** Creates the game framework source code for a new project (Pawn, GameMode, PlayerControlleR). On failure, OutFailReason will be populated. */
-	static bool GenerateGameFrameworkSourceCode(const FString& NewProjectSourcePath, const FString& NewProjectName, TArray<FString>& OutCreatedFiles, FText& OutFailReason);
+	static UE_API bool GenerateGameFrameworkSourceCode(const FString& NewProjectSourcePath, const FString& NewProjectName, TArray<FString>& OutCreatedFiles, FText& OutFailReason);
 
 	/** Creates the batch file to regenerate code projects. */
-	static bool GenerateCodeProjectGenerationBatchFile(const FString& ProjectFolder, TArray<FString>& OutCreatedFiles, FText& OutFailReason);
+	static UE_API bool GenerateCodeProjectGenerationBatchFile(const FString& ProjectFolder, TArray<FString>& OutCreatedFiles, FText& OutFailReason);
 
 	/** Creates the batch file for launching the editor or game */
-	static bool GenerateLaunchBatchFile(const FString& ProjectName, const FString& ProjectFolder, bool bLaunchEditor, TArray<FString>& OutCreatedFiles, FText& OutFailReason);
+	static UE_API bool GenerateLaunchBatchFile(const FString& ProjectName, const FString& ProjectFolder, bool bLaunchEditor, TArray<FString>& OutCreatedFiles, FText& OutFailReason);
 
 	/** Returns the copyright line used at the top of all files */
-	static FString MakeCopyrightLine();
+	static UE_API FString MakeCopyrightLine();
 
 	/** Generates a header file for a UObject class. OutSyncLocation is a string representing the preferred cursor sync location for this file after creation. */
-	static bool GenerateClassHeaderFile(const FString& NewHeaderFileName, const FString UnPrefixedClassName, const FNewClassInfo ParentClassInfo, const TArray<FString>& ClassSpecifierList, const FString& ClassProperties, const FString& ClassFunctionDeclarations, FString& OutSyncLocation, const FModuleContextInfo& ModuleInfo, bool bDeclareConstructor, FText& OutFailReason);
+	static UE_API bool GenerateClassHeaderFile(const FString& NewHeaderFileName, const FString UnPrefixedClassName, const FNewClassInfo ParentClassInfo, const TArray<FString>& ClassSpecifierList, const FString& ClassProperties, const FString& ClassFunctionDeclarations, FString& OutSyncLocation, const FModuleContextInfo& ModuleInfo, bool bDeclareConstructor, FText& OutFailReason);
 
 	/** Finds the cursor sync location in the source file and reports it back as a string */
-	static void HarvestCursorSyncLocation( FString& FinalOutput, FString& OutSyncLocation );
+	static UE_API void HarvestCursorSyncLocation( FString& FinalOutput, FString& OutSyncLocation );
 
 	/** Generates a cpp file for a UObject class */
-	static bool GenerateClassCPPFile(const FString& NewCPPFileName, const FString UnPrefixedClassName, const FNewClassInfo ParentClassInfo, const TArray<FString>& AdditionalIncludes, const TArray<FString>& PropertyOverrides, const FString& AdditionalMemberDefinitions, FString& OutSyncLocation, const FModuleContextInfo& ModuleInfo, FText& OutFailReason);
+	static UE_API bool GenerateClassCPPFile(const FString& NewCPPFileName, const FString UnPrefixedClassName, const FNewClassInfo ParentClassInfo, const TArray<FString>& AdditionalIncludes, const TArray<FString>& PropertyOverrides, const FString& AdditionalMemberDefinitions, FString& OutSyncLocation, const FModuleContextInfo& ModuleInfo, FText& OutFailReason);
 
 	/** Generates a Target.cs file for a game module */
-	static bool GenerateGameModuleTargetFile(const FString& NewTargetFileName, const FString& ModuleName, const TArray<FString>& ExtraModuleNames, FText& OutFailReason);
+	static UE_API bool GenerateGameModuleTargetFile(const FString& NewTargetFileName, const FString& ModuleName, const TArray<FString>& ExtraModuleNames, FText& OutFailReason);
 
 	/** Generates a Build.cs file for a Editor module */
-	static bool GenerateEditorModuleBuildFile(const FString& NewBuildFileName, const FString& ModuleName, const TArray<FString>& PublicDependencyModuleNames, const TArray<FString>& PrivateDependencyModuleNames, FText& OutFailReason);
+	static UE_API bool GenerateEditorModuleBuildFile(const FString& NewBuildFileName, const FString& ModuleName, const TArray<FString>& PublicDependencyModuleNames, const TArray<FString>& PrivateDependencyModuleNames, FText& OutFailReason);
 
 	/** Generates a Target.cs file for a Editor module */
-	static bool GenerateEditorModuleTargetFile(const FString& NewTargetFileName, const FString& ModuleName, const TArray<FString>& ExtraModuleNames, FText& OutFailReason);
+	static UE_API bool GenerateEditorModuleTargetFile(const FString& NewTargetFileName, const FString& ModuleName, const TArray<FString>& ExtraModuleNames, FText& OutFailReason);
 
 	/** Generates a main game module cpp file */
-	static bool GenerateGameModuleCPPFile(const FString& NewGameModuleCPPFileName, const FString& ModuleName, const FString& GameName, FText& OutFailReason);
+	static UE_API bool GenerateGameModuleCPPFile(const FString& NewGameModuleCPPFileName, const FString& ModuleName, const FString& GameName, FText& OutFailReason);
 
 	/** Generates a main game module header file */
-	static bool GenerateGameModuleHeaderFile(const FString& NewGameModuleHeaderFileName, const TArray<FString>& PublicHeaderIncludes, FText& OutFailReason);
+	static UE_API bool GenerateGameModuleHeaderFile(const FString& NewGameModuleHeaderFileName, const TArray<FString>& PublicHeaderIncludes, FText& OutFailReason);
 
 	/**
 	* Replace a wildcard with another string
@@ -423,35 +425,35 @@ private:
 
 	* @return the input string after replacement of the wildcard
 	*/
-	static FString ReplaceWildcard(const FString& Input, const FString& From, const FString& To, bool bLeadingTab = false, bool bTrailingNewLine = false);
+	static UE_API FString ReplaceWildcard(const FString& Input, const FString& From, const FString& To, bool bLeadingTab = false, bool bTrailingNewLine = false);
 
 	/** Handler for when the user confirms a project update */
-	static void OnUpdateProjectConfirm();
+	static UE_API void OnUpdateProjectConfirm();
 
 	/** @param OutProjectCodeFilenames Contains the filenames of the project source code files */
-	static void GetProjectCodeFilenames(TArray<FString>& OutProjectCodeFilenames);
+	static UE_API void GetProjectCodeFilenames(TArray<FString>& OutProjectCodeFilenames);
 
 	/**
 	 * Updates the projects and modifies FProjectDescriptor accordingly to given modifier.
 	 *
 	 * @param Modifier	Callback delegate that will modify the project descriptor accordingly.
 	 */
-	static void UpdateProject(const FProjectDescriptorModifier& Modifier);
+	static UE_API void UpdateProject(const FProjectDescriptorModifier& Modifier);
 
 	/**
 	 * Updates the project file.
 	 */
-	static void UpdateProject();
+	static UE_API void UpdateProject();
 
 	/**
 	 * Updates the projects, and optionally the modules names
 	 *
 	 * @param StartupModuleNames	if specified, replaces the existing module names with this version
 	 */
-	static void UpdateProject(const TArray<FString>* StartupModuleNames);
+	static UE_API void UpdateProject(const TArray<FString>* StartupModuleNames);
 
 	/** Handler for when the user opts out of a project update */
-	static void OnUpdateProjectCancel();
+	static UE_API void OnUpdateProjectCancel();
 
 	/**
 	 * Updates the loaded game project file to the current version and to use the given modules
@@ -463,7 +465,7 @@ private:
 
 	 * @return true, if successful
 	 */
-	static bool UpdateGameProjectFile(const FString& ProjectFilename, const FString& EngineIdentifier, const TArray<FString>* StartupModuleNames, FText& OutFailReason);
+	static UE_API bool UpdateGameProjectFile(const FString& ProjectFilename, const FString& EngineIdentifier, const TArray<FString>* StartupModuleNames, FText& OutFailReason);
 
 	/**
 	 * Updates the loaded game project file to the current version and modifies FProjectDescriptor accordingly to given modifier.
@@ -475,7 +477,7 @@ private:
 
 	 * @return true, if successful
 	 */
-	static bool UpdateGameProjectFile(const FString& ProjectFilename, const FString& EngineIdentifier, const FProjectDescriptorModifier& Modifier, FText& OutFailReason);
+	static UE_API bool UpdateGameProjectFile(const FString& ProjectFilename, const FString& EngineIdentifier, const FProjectDescriptorModifier& Modifier, FText& OutFailReason);
 
 	/**
 	 * Updates the loaded game project file to the current version.
@@ -486,23 +488,23 @@ private:
 
 	 * @return true, if successful
 	 */
-	static bool UpdateGameProjectFile(const FString& ProjectFilename, const FString& EngineIdentifier, FText& OutFailReason);
+	static UE_API bool UpdateGameProjectFile(const FString& ProjectFilename, const FString& EngineIdentifier, FText& OutFailReason);
 
 	/** Checks the specified game project file out from source control */
-	static bool CheckoutGameProjectFile(const FString& ProjectFilename, FText& OutFailReason);
+	static UE_API bool CheckoutGameProjectFile(const FString& ProjectFilename, FText& OutFailReason);
 
 	/** Internal handler for AddCodeToProject*/
-	static EAddCodeToProjectResult AddCodeToProject_Internal(const FString& NewClassName, const FString& NewClassPath, const FModuleContextInfo& ModuleInfo, const FNewClassInfo ParentClassInfo, const TSet<FString>& DisallowedHeaderNames, FString& OutHeaderFilePath, FString& OutCppFilePath, FText& OutFailReason, EReloadStatus& OutReloadStatus);
+	static UE_API EAddCodeToProjectResult AddCodeToProject_Internal(const FString& NewClassName, const FString& NewClassPath, const FModuleContextInfo& ModuleInfo, const FNewClassInfo ParentClassInfo, const TSet<FString>& DisallowedHeaderNames, FString& OutHeaderFilePath, FString& OutCppFilePath, FText& OutFailReason, EReloadStatus& OutReloadStatus);
 
 	/** Internal handler for IsValidBaseClassForCreation */
 	DECLARE_DELEGATE_RetVal_OneParam(bool, FDoesClassNeedAPIExportCallback, const FString& /*ClassModuleName*/);
-	static bool IsValidBaseClassForCreation_Internal(const UClass* InClass, const FDoesClassNeedAPIExportCallback& InDoesClassNeedAPIExport);
+	static UE_API bool IsValidBaseClassForCreation_Internal(const UClass* InClass, const FDoesClassNeedAPIExportCallback& InDoesClassNeedAPIExport);
 
 	/** Handler for the user confirming they've read the name legnth warning */
-	static void OnWarningReasonOk();
+	static UE_API void OnWarningReasonOk();
 
 	/** Given a source file name, find its location within the project */
-	static bool FindSourceFileInProject(const FString& InFilename, const FString& InSearchPath, FString& OutPath);
+	static UE_API bool FindSourceFileInProject(const FString& InFilename, const FString& InSearchPath, FString& OutPath);
 
 	/**
 	 * Gets required additional dependencies for fresh project for given class info.
@@ -511,7 +513,7 @@ private:
 	 *
 	 * @returns Array of required dependencies.
 	 */
-	static TArray<FString> GetRequiredAdditionalDependencies(const FNewClassInfo& ClassInfo);
+	static UE_API TArray<FString> GetRequiredAdditionalDependencies(const FNewClassInfo& ClassInfo);
 
 	/**
 	 * Updates startup module names in project descriptor.
@@ -521,7 +523,7 @@ private:
 	 *
 	 * @returns True if descriptor has been modified. False otherwise.
 	 */
-	static bool UpdateStartupModuleNames(FProjectDescriptor& Descriptor, const TArray<FString>* StartupModuleNames);
+	static UE_API bool UpdateStartupModuleNames(FProjectDescriptor& Descriptor, const TArray<FString>* StartupModuleNames);
 
 	/**
 	 * Updates additional dependencies in project descriptor.
@@ -532,14 +534,14 @@ private:
 	 *
 	 * @returns True if descriptor has been modified. False otherwise.
 	 */
-	static bool UpdateRequiredAdditionalDependencies(FProjectDescriptor& Descriptor, TArray<FString>& RequiredDependencies, const FString& ModuleName);
+	static UE_API bool UpdateRequiredAdditionalDependencies(FProjectDescriptor& Descriptor, TArray<FString>& RequiredDependencies, const FString& ModuleName);
 
 	/**
 	 * Updates the projects and modifies FProjectDescriptor accordingly to given modifier.
 	 *
 	 * @param Modifier	Callback delegate that will modify the project descriptor accordingly.
 	 */
-	static void UpdateProject_Impl(const FProjectDescriptorModifier* Modifier);
+	static UE_API void UpdateProject_Impl(const FProjectDescriptorModifier* Modifier);
 
 	/**
 	 * Updates the loaded game project file to the current version and modifies FProjectDescriptor accordingly to given modifier.
@@ -551,15 +553,17 @@ private:
 
 	 * @return true, if successful
 	 */
-	static bool UpdateGameProjectFile_Impl(const FString& ProjectFilename, const FString& EngineIdentifier, const FProjectDescriptorModifier* Modifier, FText& OutFailReason);
+	static UE_API bool UpdateGameProjectFile_Impl(const FString& ProjectFilename, const FString& EngineIdentifier, const FProjectDescriptorModifier* Modifier, FText& OutFailReason);
 
 private:
 
-	static TWeakPtr<SNotificationItem> UpdateGameProjectNotification;
-	static TWeakPtr<SNotificationItem> WarningProjectNameNotification;
+	static UE_API TWeakPtr<SNotificationItem> UpdateGameProjectNotification;
+	static UE_API TWeakPtr<SNotificationItem> WarningProjectNameNotification;
 
 	// Whether we should use AudioMixer for all platforms:
-	static bool bUseAudioMixerForAllPlatforms;
+	static UE_API bool bUseAudioMixerForAllPlatforms;
 
 	constexpr static const TCHAR IncludePathFormatString[] = TEXT("#include \"%s\"");
 };
+
+#undef UE_API
diff --git a/Engine/Source/Editor/GameplayTasksEditor/Classes/K2Node_LatentGameplayTaskCall.h b/Engine/Source/Editor/GameplayTasksEditor/Classes/K2Node_LatentGameplayTaskCall.h
--- a/Engine/Source/Editor/GameplayTasksEditor/Classes/K2Node_LatentGameplayTaskCall.h
+++ b/Engine/Source/Editor/GameplayTasksEditor/Classes/K2Node_LatentGameplayTaskCall.h
@@ -14,6 +14,8 @@
 
 #include "K2Node_LatentGameplayTaskCall.generated.h"
 
+#define UE_API GAMEPLAYTASKSEDITOR_API
+
 class FBlueprintActionDatabaseRegistrar;
 class FName;
 class UClass;
@@ -24,41 +26,43 @@ class UEdGraphSchema_K2;
 class UGameplayTask;
 class UObject;
 
-UCLASS()
-class GAMEPLAYTASKSEDITOR_API UK2Node_LatentGameplayTaskCall : public UK2Node_BaseAsyncTask
+UCLASS(MinimalAPI)
+class UK2Node_LatentGameplayTaskCall : public UK2Node_BaseAsyncTask
 {
 	GENERATED_BODY()
 
 public:
-	UK2Node_LatentGameplayTaskCall(const FObjectInitializer& ObjectInitializer);
+	UE_API UK2Node_LatentGameplayTaskCall(const FObjectInitializer& ObjectInitializer);
 
 	// UEdGraphNode interface
-	virtual bool CanCreateUnderSpecifiedSchema(const UEdGraphSchema* DesiredSchema) const override;
-	virtual void GetMenuActions(FBlueprintActionDatabaseRegistrar& ActionRegistrar) const override;
+	UE_API virtual bool CanCreateUnderSpecifiedSchema(const UEdGraphSchema* DesiredSchema) const override;
+	UE_API virtual void GetMenuActions(FBlueprintActionDatabaseRegistrar& ActionRegistrar) const override;
 	// End of UEdGraphNode interface
 
-	virtual void ReallocatePinsDuringReconstruction(TArray<UEdGraphPin*>& OldPins) override;
-	virtual void PinDefaultValueChanged(UEdGraphPin* Pin) override;
-	bool ConnectSpawnProperties(UClass* ClassToSpawn, const UEdGraphSchema_K2* Schema, class FKismetCompilerContext& CompilerContext, UEdGraph* SourceGraph, UEdGraphPin*& LastThenPin, UEdGraphPin* SpawnedActorReturnPin);		//Helper
-	virtual void ExpandNode(class FKismetCompilerContext& CompilerContext, UEdGraph* SourceGraph) override;
+	UE_API virtual void ReallocatePinsDuringReconstruction(TArray<UEdGraphPin*>& OldPins) override;
+	UE_API virtual void PinDefaultValueChanged(UEdGraphPin* Pin) override;
+	UE_API bool ConnectSpawnProperties(UClass* ClassToSpawn, const UEdGraphSchema_K2* Schema, class FKismetCompilerContext& CompilerContext, UEdGraph* SourceGraph, UEdGraphPin*& LastThenPin, UEdGraphPin* SpawnedActorReturnPin);		//Helper
+	UE_API virtual void ExpandNode(class FKismetCompilerContext& CompilerContext, UEdGraph* SourceGraph) override;
 
-	void CreatePinsForClass(UClass* InClass);
-	UEdGraphPin* GetClassPin(const TArray<UEdGraphPin*>* InPinsToSearch = NULL) const;
-	UClass* GetClassToSpawn(const TArray<UEdGraphPin*>* InPinsToSearch = NULL) const;
-	UEdGraphPin* GetResultPin() const;
-	bool IsSpawnVarPin(UEdGraphPin* Pin);
-	bool ValidateActorSpawning(class FKismetCompilerContext& CompilerContext, bool bGenerateErrors);
-	bool ValidateActorArraySpawning(class FKismetCompilerContext& CompilerContext, bool bGenerateErrors);
+	UE_API void CreatePinsForClass(UClass* InClass);
+	UE_API UEdGraphPin* GetClassPin(const TArray<UEdGraphPin*>* InPinsToSearch = NULL) const;
+	UE_API UClass* GetClassToSpawn(const TArray<UEdGraphPin*>* InPinsToSearch = NULL) const;
+	UE_API UEdGraphPin* GetResultPin() const;
+	UE_API bool IsSpawnVarPin(UEdGraphPin* Pin);
+	UE_API bool ValidateActorSpawning(class FKismetCompilerContext& CompilerContext, bool bGenerateErrors);
+	UE_API bool ValidateActorArraySpawning(class FKismetCompilerContext& CompilerContext, bool bGenerateErrors);
 
 	UPROPERTY()
 	TArray<FName> SpawnParamPins;
 
-	static void RegisterSpecializedTaskNodeClass(TSubclassOf<UK2Node_LatentGameplayTaskCall> NodeClass);
+	static UE_API void RegisterSpecializedTaskNodeClass(TSubclassOf<UK2Node_LatentGameplayTaskCall> NodeClass);
 protected:
-	static bool HasDedicatedNodeClass(TSubclassOf<UGameplayTask> TaskClass);
+	static UE_API bool HasDedicatedNodeClass(TSubclassOf<UGameplayTask> TaskClass);
 
 	virtual bool IsHandling(TSubclassOf<UGameplayTask> TaskClass) const { return true; }
 
 private:
-	static TArray<TWeakObjectPtr<UClass> > NodeClasses;
+	static UE_API TArray<TWeakObjectPtr<UClass> > NodeClasses;
 };
+
+#undef UE_API
diff --git a/Engine/Source/Editor/GraphEditor/Private/DragNode.h b/Engine/Source/Editor/GraphEditor/Private/DragNode.h
--- a/Engine/Source/Editor/GraphEditor/Private/DragNode.h
+++ b/Engine/Source/Editor/GraphEditor/Private/DragNode.h
@@ -10,37 +10,39 @@
 #include "Math/Vector2D.h"
 #include "Templates/SharedPointer.h"
 
+#define UE_API GRAPHEDITOR_API
+
 class SGraphNode;
 class SGraphPanel;
 class SWidget;
 class UEdGraph;
 class UEdGraphNode;
 
-class GRAPHEDITOR_API FDragNode : public FGraphEditorDragDropAction
+class FDragNode : public FGraphEditorDragDropAction
 {
 public:
 	DRAG_DROP_OPERATOR_TYPE(FDragNode, FGraphEditorDragDropAction)
 
-	static TSharedRef<FDragNode> New(const TSharedRef<SGraphPanel>& InGraphPanel, const TSharedRef<SGraphNode>& InDraggedNode);
+	static UE_API TSharedRef<FDragNode> New(const TSharedRef<SGraphPanel>& InGraphPanel, const TSharedRef<SGraphNode>& InDraggedNode);
 
-	static TSharedRef<FDragNode> New(const TSharedRef<SGraphPanel>& InGraphPanel, const TArray< TSharedRef<SGraphNode> >& InDraggedNodes);
+	static UE_API TSharedRef<FDragNode> New(const TSharedRef<SGraphPanel>& InGraphPanel, const TArray< TSharedRef<SGraphNode> >& InDraggedNodes);
 	
 	// FGraphEditorDragDropAction interface
-	virtual void HoverTargetChanged() override;
-	virtual FReply DroppedOnNode(const FVector2f& ScreenPosition, const FVector2f& GraphPosition) override;
-	virtual FReply DroppedOnPanel( const TSharedRef< SWidget >& Panel, const FVector2f& ScreenPosition, const FVector2f& GraphPosition, UEdGraph& Graph) override;
+	UE_API virtual void HoverTargetChanged() override;
+	UE_API virtual FReply DroppedOnNode(const FVector2f& ScreenPosition, const FVector2f& GraphPosition) override;
+	UE_API virtual FReply DroppedOnPanel( const TSharedRef< SWidget >& Panel, const FVector2f& ScreenPosition, const FVector2f& GraphPosition, UEdGraph& Graph) override;
 	//virtual void OnDragBegin( const TSharedRef<class SGraphPin>& InPin) override;
-	virtual void OnDragged (const class FDragDropEvent& DragDropEvent ) override;
+	UE_API virtual void OnDragged (const class FDragDropEvent& DragDropEvent ) override;
 	// End of FGraphEditorDragDropAction interface
 	
-	const TArray< TSharedRef<SGraphNode> > & GetNodes() const;
+	UE_API const TArray< TSharedRef<SGraphNode> > & GetNodes() const;
 
-	virtual bool IsValidOperation() const;
+	UE_API virtual bool IsValidOperation() const;
 
 protected:
 	typedef FGraphEditorDragDropAction Super;
 
-	UEdGraphNode* GetGraphNodeForSGraphNode(TSharedRef<SGraphNode>& SNode);
+	UE_API UEdGraphNode* GetGraphNodeForSGraphNode(TSharedRef<SGraphNode>& SNode);
 
 protected:
 
@@ -56,3 +58,5 @@ protected:
 	/** if we can drop our node here*/
 	bool bValidOperation;
 };
+
+#undef UE_API
diff --git a/Engine/Source/Editor/GraphEditor/Private/KismetNodes/SGraphNodeK2Copy.h b/Engine/Source/Editor/GraphEditor/Private/KismetNodes/SGraphNodeK2Copy.h
--- a/Engine/Source/Editor/GraphEditor/Private/KismetNodes/SGraphNodeK2Copy.h
+++ b/Engine/Source/Editor/GraphEditor/Private/KismetNodes/SGraphNodeK2Copy.h
@@ -6,25 +6,29 @@
 #include "Templates/SharedPointer.h"
 #include "Widgets/DeclarativeSyntaxSupport.h"
 
+#define UE_API GRAPHEDITOR_API
+
 class SGraphPin;
 class SWidget;
 class UEdGraphPin;
 class UK2Node;
 struct FSlateBrush;
 
-class GRAPHEDITOR_API SGraphNodeK2Copy : public SGraphNodeK2Base
+class SGraphNodeK2Copy : public SGraphNodeK2Base
 {
 public:
 	SLATE_BEGIN_ARGS(SGraphNodeK2Copy) {}
 	SLATE_END_ARGS()
 
-	void Construct(const FArguments& InArgs, UK2Node* InNode);
+	UE_API void Construct(const FArguments& InArgs, UK2Node* InNode);
 
 	//~ Begin SGraphNode Interface
-	virtual void UpdateGraphNode() override;
-	virtual TSharedRef<SWidget> CreateNodeContentArea() override;
-	virtual void AddPin(const TSharedRef<SGraphPin>& PinToAdd) override;
-	virtual const FSlateBrush* GetShadowBrush(bool bSelected) const override;
-	virtual TSharedPtr<SGraphPin> CreatePinWidget(UEdGraphPin* Pin) const override;
+	UE_API virtual void UpdateGraphNode() override;
+	UE_API virtual TSharedRef<SWidget> CreateNodeContentArea() override;
+	UE_API virtual void AddPin(const TSharedRef<SGraphPin>& PinToAdd) override;
+	UE_API virtual const FSlateBrush* GetShadowBrush(bool bSelected) const override;
+	UE_API virtual TSharedPtr<SGraphPin> CreatePinWidget(UEdGraphPin* Pin) const override;
 	//~ End SGraphNode Interface
 };
+
+#undef UE_API
diff --git a/Engine/Source/Editor/GraphEditor/Public/BlueprintConnectionDrawingPolicy.h b/Engine/Source/Editor/GraphEditor/Public/BlueprintConnectionDrawingPolicy.h
--- a/Engine/Source/Editor/GraphEditor/Public/BlueprintConnectionDrawingPolicy.h
+++ b/Engine/Source/Editor/GraphEditor/Public/BlueprintConnectionDrawingPolicy.h
@@ -16,6 +16,8 @@
 #include "Templates/SharedPointer.h"
 #include "Widgets/SWidget.h"
 
+#define UE_API GRAPHEDITOR_API
+
 class FArrangedChildren;
 class FArrangedWidget;
 class FSlateRect;
@@ -30,7 +32,7 @@ class UEdGraphPin;
 // FKismetConnectionDrawingPolicy
 
 // This class draws the connections for an UEdGraph using a K2-based schema
-class GRAPHEDITOR_API FKismetConnectionDrawingPolicy : public FConnectionDrawingPolicy
+class FKismetConnectionDrawingPolicy : public FConnectionDrawingPolicy
 {
 protected:
 	// Times for one execution pair within the current graph
@@ -83,32 +85,32 @@ protected:
 
 protected:
 	// Should this wire be treated as an execution pin, querying the execution trace for connected nodes to draw it differently?
-	virtual bool TreatWireAsExecutionPin(UEdGraphPin* InputPin, UEdGraphPin* OutputPin) const;
+	UE_API virtual bool TreatWireAsExecutionPin(UEdGraphPin* InputPin, UEdGraphPin* OutputPin) const;
 
 	// Determines the target Blueprint for querying the active debugging context.
-	UBlueprint* GetTargetBlueprint() const;
+	UE_API UBlueprint* GetTargetBlueprint() const;
 
 	// Determines if the given Blueprint context has a valid object for debugging.
-	bool CanBuildRoadmap(UBlueprint* TargetBP) const;
+	UE_API bool CanBuildRoadmap(UBlueprint* TargetBP) const;
 
 public:
-	FKismetConnectionDrawingPolicy(int32 InBackLayerID, int32 InFrontLayerID, float ZoomFactor, const FSlateRect& InClippingRect, FSlateWindowElementList& InDrawElements, UEdGraph* InGraphObj);
+	UE_API FKismetConnectionDrawingPolicy(int32 InBackLayerID, int32 InFrontLayerID, float ZoomFactor, const FSlateRect& InClippingRect, FSlateWindowElementList& InDrawElements, UEdGraph* InGraphObj);
 
-	virtual void BuildExecutionRoadmap();
+	UE_API virtual void BuildExecutionRoadmap();
 
-	virtual bool CanBuildRoadmap() const;
+	UE_API virtual bool CanBuildRoadmap() const;
 
-	void CalculateEnvelopeAlphas(double ExecutionTime, /*out*/ float& AttackAlpha, /*out*/ float& SustainAlpha) const;
+	UE_API void CalculateEnvelopeAlphas(double ExecutionTime, /*out*/ float& AttackAlpha, /*out*/ float& SustainAlpha) const;
 
 	// FConnectionDrawingPolicy interface
-	virtual void DetermineWiringStyle(UEdGraphPin* OutputPin, UEdGraphPin* InputPin, /*inout*/ FConnectionParams& Params) override;
-	virtual void Draw(TMap<TSharedRef<SWidget>, FArrangedWidget>& PinGeometries, FArrangedChildren& ArrangedNodes) override;
+	UE_API virtual void DetermineWiringStyle(UEdGraphPin* OutputPin, UEdGraphPin* InputPin, /*inout*/ FConnectionParams& Params) override;
+	UE_API virtual void Draw(TMap<TSharedRef<SWidget>, FArrangedWidget>& PinGeometries, FArrangedChildren& ArrangedNodes) override;
 	// End of FConnectionDrawingPolicy interface
 
-	virtual void DetermineStyleOfExecWire(float& Thickness, FLinearColor& WireColor, bool& bDrawBubbles, const FTimePair& Times);
+	UE_API virtual void DetermineStyleOfExecWire(float& Thickness, FLinearColor& WireColor, bool& bDrawBubbles, const FTimePair& Times);
 
-	virtual void SetIncompatiblePinDrawState(const TSharedPtr<SGraphPin>& StartPin, const TSet< TSharedRef<SWidget> >& VisiblePins) override;
-	virtual void ResetIncompatiblePinDrawState(const TSet< TSharedRef<SWidget> >& VisiblePins) override;
+	UE_API virtual void SetIncompatiblePinDrawState(const TSharedPtr<SGraphPin>& StartPin, const TSet< TSharedRef<SWidget> >& VisiblePins) override;
+	UE_API virtual void ResetIncompatiblePinDrawState(const TSet< TSharedRef<SWidget> >& VisiblePins) override;
 
 protected: 
 	/**
@@ -119,7 +121,7 @@ protected:
 	* @param  OutputNode	The node whose exec input should be coming from the InputNode.
 	* @return True if the two nodes are connected in the graph, otherwise false.
 	*/
-	bool AreNodesGraphicallySequential(UEdGraphNode* InputNode, UEdGraphNode* OutputNode) const;
+	UE_API bool AreNodesGraphicallySequential(UEdGraphNode* InputNode, UEdGraphNode* OutputNode) const;
 
 	/**
 	 * With the introduction of "reroute" (UK2Node_Knot) nodes, we have to 
@@ -135,10 +137,12 @@ protected:
 	 * @param  NodeExecutionList	A list of pins that feed the execution of a specific node.
 	 * @return Null if OutputPin did not cause the invocation of the node in question, otherwise a valid time pair denoting when the node was executed.
 	 */
-	FTimePair const* BackTraceExecPath(UEdGraphPin const* const OutputPin, FExecPairingMap const* const NodeExecutionList);
+	UE_API FTimePair const* BackTraceExecPath(UEdGraphPin const* const OutputPin, FExecPairingMap const* const NodeExecutionList);
 
-	bool ShouldChangeTangentForKnot(class UK2Node_Knot* Knot);
-	bool GetAverageConnectedPosition(class UK2Node_Knot* Knot, EEdGraphPinDirection Direction, FVector2f& OutPos) const;
-	bool FindPinCenter(UEdGraphPin* Pin, FVector2f& OutCenter) const;
-	static UObject* GetObjectBeingDebugged(UBlueprint* TargetBP);
+	UE_API bool ShouldChangeTangentForKnot(class UK2Node_Knot* Knot);
+	UE_API bool GetAverageConnectedPosition(class UK2Node_Knot* Knot, EEdGraphPinDirection Direction, FVector2f& OutPos) const;
+	UE_API bool FindPinCenter(UEdGraphPin* Pin, FVector2f& OutCenter) const;
+	static UE_API UObject* GetObjectBeingDebugged(UBlueprint* TargetBP);
 };
+
+#undef UE_API
diff --git a/Engine/Source/Editor/GraphEditor/Public/ConnectionDrawingPolicy.h b/Engine/Source/Editor/GraphEditor/Public/ConnectionDrawingPolicy.h
--- a/Engine/Source/Editor/GraphEditor/Public/ConnectionDrawingPolicy.h
+++ b/Engine/Source/Editor/GraphEditor/Public/ConnectionDrawingPolicy.h
@@ -19,6 +19,8 @@
 #include "Templates/SharedPointer.h"
 #include "Widgets/SWidget.h"
 
+#define UE_API GRAPHEDITOR_API
+
 class FArrangedChildren;
 class FArrangedWidget;
 class FSlateRect;
@@ -37,27 +39,27 @@ GRAPHEDITOR_API DECLARE_LOG_CATEGORY_EXTERN(LogConnectionDrawingPolicy, Log, All
 /////////////////////////////////////////////////////
 // FGeometryHelper
 
-class GRAPHEDITOR_API FGeometryHelper
+class FGeometryHelper
 {
 public:
-	static UE::Slate::FDeprecateVector2DResult VerticalMiddleLeftOf(const FGeometry& SomeGeometry);
-	static UE::Slate::FDeprecateVector2DResult VerticalMiddleRightOf(const FGeometry& SomeGeometry);
+	static UE_API UE::Slate::FDeprecateVector2DResult VerticalMiddleLeftOf(const FGeometry& SomeGeometry);
+	static UE_API UE::Slate::FDeprecateVector2DResult VerticalMiddleRightOf(const FGeometry& SomeGeometry);
 
-	static UE::Slate::FDeprecateVector2DResult CenterOf(const FGeometry& SomeGeometry);
+	static UE_API UE::Slate::FDeprecateVector2DResult CenterOf(const FGeometry& SomeGeometry);
 	UE_DEPRECATED(5.6, "Use Array of FVector2f instead of FVector2D for slate positions")
-	static void ConvertToPoints(const FGeometry& Geom, TArray<FVector2D>& Points);
-	static void ConvertToPoints(const FGeometry& Geom, TArray<FVector2f>& Points);
+	static UE_API void ConvertToPoints(const FGeometry& Geom, TArray<FVector2D>& Points);
+	static UE_API void ConvertToPoints(const FGeometry& Geom, TArray<FVector2f>& Points);
 
 	/** Find the point on line segment from LineStart to LineEnd which is closest to Point */
-	static UE::Slate::FDeprecateVector2DResult FindClosestPointOnLine(const UE::Slate::FDeprecateVector2DParameter& LineStart, const UE::Slate::FDeprecateVector2DParameter& LineEnd, const UE::Slate::FDeprecateVector2DParameter& TestPoint);
+	static UE_API UE::Slate::FDeprecateVector2DResult FindClosestPointOnLine(const UE::Slate::FDeprecateVector2DParameter& LineStart, const UE::Slate::FDeprecateVector2DParameter& LineEnd, const UE::Slate::FDeprecateVector2DParameter& TestPoint);
 
-	static UE::Slate::FDeprecateVector2DResult FindClosestPointOnGeom(const FGeometry& Geom, const UE::Slate::FDeprecateVector2DParameter& TestPoint);
+	static UE_API UE::Slate::FDeprecateVector2DResult FindClosestPointOnGeom(const FGeometry& Geom, const UE::Slate::FDeprecateVector2DParameter& TestPoint);
 };
 
 /////////////////////////////////////////////////////
 // FConnectionParameters
 
-struct GRAPHEDITOR_API FConnectionParams
+struct FConnectionParams
 {
 	FLinearColor WireColor;
 	UEdGraphPin* AssociatedPin1;
@@ -96,7 +98,7 @@ struct GRAPHEDITOR_API FConnectionParams
 // FConnectionDrawingPolicy
 
 // This class draws the connections for an UEdGraph composed of pins and nodes
-class GRAPHEDITOR_API FConnectionDrawingPolicy
+class FConnectionDrawingPolicy
 {
 protected:
 	int32 WireLayerID;
@@ -140,15 +142,15 @@ protected:
 public:
 	virtual ~FConnectionDrawingPolicy() {}
 	
-	FConnectionDrawingPolicy(int32 InBackLayerID, int32 InFrontLayerID, float InZoomFactor, const FSlateRect& InClippingRect, FSlateWindowElementList& InDrawElements);
+	UE_API FConnectionDrawingPolicy(int32 InBackLayerID, int32 InFrontLayerID, float InZoomFactor, const FSlateRect& InClippingRect, FSlateWindowElementList& InDrawElements);
 
 	// Update the drawing policy with the set of hovered pins (which can be empty)
-	void SetHoveredPins(const TSet< FEdGraphPinReference >& InHoveredPins, const TArray< TSharedPtr<SGraphPin> >& OverridePins, double HoverTime);
+	UE_API void SetHoveredPins(const TSet< FEdGraphPinReference >& InHoveredPins, const TArray< TSharedPtr<SGraphPin> >& OverridePins, double HoverTime);
 
-	void SetMousePosition(const UE::Slate::FDeprecateVector2DParameter& InMousePos);
+	UE_API void SetMousePosition(const UE::Slate::FDeprecateVector2DParameter& InMousePos);
 
 	// Update the drawing policy with the marked pin (which may not be valid)
-	void SetMarkedPin(TWeakPtr<SGraphPin> InMarkedPin);
+	UE_API void SetMarkedPin(TWeakPtr<SGraphPin> InMarkedPin);
 
 	// Set the selected nodes from the graph panel.
 	void SetSelectedNodes(const TArray<UEdGraphNode*>& InSelectedNodes) { SelectedGraphNodes = InSelectedNodes; }
@@ -156,33 +158,33 @@ public:
 	// Set the list of currently relinked connections.
 	void SetRelinkConnections(const TArray<FRelinkConnection>& Connections) { RelinkConnections = Connections; }
 
-	static float MakeSplineReparamTable(const UE::Slate::FDeprecateVector2DParameter& P0, const UE::Slate::FDeprecateVector2DParameter& P0Tangent, const UE::Slate::FDeprecateVector2DParameter& P1, const UE::Slate::FDeprecateVector2DParameter& P1Tangent, FInterpCurve<float>& OutReparamTable);
+	static UE_API float MakeSplineReparamTable(const UE::Slate::FDeprecateVector2DParameter& P0, const UE::Slate::FDeprecateVector2DParameter& P0Tangent, const UE::Slate::FDeprecateVector2DParameter& P1, const UE::Slate::FDeprecateVector2DParameter& P1Tangent, FInterpCurve<float>& OutReparamTable);
 
 	UE_DEPRECATED(5.6, "Use the version of the function accepting FVector2f; this Slate API no longer interfaces directly with double-precision scalars and vectors.")
-	virtual void DrawSplineWithArrow(const FVector2D& StartPoint, const FVector2D& EndPoint, const FConnectionParams& Params) UE_SLATE_DEPRECATED_VECTOR_VIRTUAL_FUNCTION;
-	virtual void DrawSplineWithArrow(const FVector2f& StartPoint, const FVector2f& EndPoint, const FConnectionParams& Params);
+	UE_API virtual void DrawSplineWithArrow(const FVector2D& StartPoint, const FVector2D& EndPoint, const FConnectionParams& Params) UE_SLATE_DEPRECATED_VECTOR_VIRTUAL_FUNCTION;
+	UE_API virtual void DrawSplineWithArrow(const FVector2f& StartPoint, const FVector2f& EndPoint, const FConnectionParams& Params);
 	
-	virtual void DrawSplineWithArrow(const FGeometry& StartGeom, const FGeometry& EndGeom, const FConnectionParams& Params);
+	UE_API virtual void DrawSplineWithArrow(const FGeometry& StartGeom, const FGeometry& EndGeom, const FConnectionParams& Params);
 
 	UE_DEPRECATED(5.6, "Use the version of the function accepting FVector2f; this Slate API no longer interfaces directly with double-precision scalars and vectors.")
-	virtual FVector2D ComputeSplineTangent(const FVector2D& Start, const FVector2D& End) const UE_SLATE_DEPRECATED_VECTOR_VIRTUAL_FUNCTION;
-	virtual FVector2f ComputeSplineTangent(const FVector2f& Start, const FVector2f& End) const;
+	UE_API virtual FVector2D ComputeSplineTangent(const FVector2D& Start, const FVector2D& End) const UE_SLATE_DEPRECATED_VECTOR_VIRTUAL_FUNCTION;
+	UE_API virtual FVector2f ComputeSplineTangent(const FVector2f& Start, const FVector2f& End) const;
 
 
 	UE_DEPRECATED(5.6, "Use the version of the function accepting FVector2f; this Slate API no longer interfaces directly with double-precision scalars and vectors.")
-	virtual void DrawConnection(int32 LayerId, const FVector2D& Start, const FVector2D& End, const FConnectionParams& Params) UE_SLATE_DEPRECATED_VECTOR_VIRTUAL_FUNCTION;
-	virtual void DrawConnection(int32 LayerId, const FVector2f& Start, const FVector2f& End, const FConnectionParams& Params);
+	UE_API virtual void DrawConnection(int32 LayerId, const FVector2D& Start, const FVector2D& End, const FConnectionParams& Params) UE_SLATE_DEPRECATED_VECTOR_VIRTUAL_FUNCTION;
+	UE_API virtual void DrawConnection(int32 LayerId, const FVector2f& Start, const FVector2f& End, const FConnectionParams& Params);
 
 	UE_DEPRECATED(5.6, "Use the version of the function accepting FVector2f; this Slate API no longer interfaces directly with double-precision scalars and vectors.")
-	virtual void DrawPreviewConnector(const FGeometry& PinGeometry, const FVector2D& StartPoint, const FVector2D& EndPoint, UEdGraphPin* Pin) UE_SLATE_DEPRECATED_VECTOR_VIRTUAL_FUNCTION;
-	virtual void DrawPreviewConnector(const FGeometry& PinGeometry, const FVector2f& StartPoint, const FVector2f& EndPoint, UEdGraphPin* Pin);
+	UE_API virtual void DrawPreviewConnector(const FGeometry& PinGeometry, const FVector2D& StartPoint, const FVector2D& EndPoint, UEdGraphPin* Pin) UE_SLATE_DEPRECATED_VECTOR_VIRTUAL_FUNCTION;
+	UE_API virtual void DrawPreviewConnector(const FGeometry& PinGeometry, const FVector2f& StartPoint, const FVector2f& EndPoint, UEdGraphPin* Pin);
 
 	// Give specific editor modes a chance to highlight this connection or darken non-interesting connections
-	virtual void DetermineWiringStyle(UEdGraphPin* OutputPin, UEdGraphPin* InputPin, /*inout*/ FConnectionParams& Params);
+	UE_API virtual void DetermineWiringStyle(UEdGraphPin* OutputPin, UEdGraphPin* InputPin, /*inout*/ FConnectionParams& Params);
 
-	virtual void Draw(TMap<TSharedRef<SWidget>, FArrangedWidget>& InPinGeometries, FArrangedChildren& ArrangedNodes);
+	UE_API virtual void Draw(TMap<TSharedRef<SWidget>, FArrangedWidget>& InPinGeometries, FArrangedChildren& ArrangedNodes);
 
-	virtual void DetermineLinkGeometry(
+	UE_API virtual void DetermineLinkGeometry(
 		FArrangedChildren& ArrangedNodes, 
 		TSharedRef<SWidget>& OutputPinWidget,
 		UEdGraphPin* OutputPin,
@@ -194,26 +196,28 @@ public:
 	// Choose whether we want to cache the pins draw state to avoid resetting it for every tick 
 	virtual bool UseDrawStateCaching() const { return false; }
 	
-	virtual void SetIncompatiblePinDrawState(const TSharedPtr<SGraphPin>& StartPin, const TSet< TSharedRef<SWidget> >& VisiblePins);
-	virtual void ResetIncompatiblePinDrawState(const TSet< TSharedRef<SWidget> >& VisiblePins);
+	UE_API virtual void SetIncompatiblePinDrawState(const TSharedPtr<SGraphPin>& StartPin, const TSet< TSharedRef<SWidget> >& VisiblePins);
+	UE_API virtual void ResetIncompatiblePinDrawState(const TSet< TSharedRef<SWidget> >& VisiblePins);
 
-	virtual void ApplyHoverDeemphasis(UEdGraphPin* OutputPin, UEdGraphPin* InputPin, /*inout*/ float& Thickness, /*inout*/ FLinearColor& WireColor);
+	UE_API virtual void ApplyHoverDeemphasis(UEdGraphPin* OutputPin, UEdGraphPin* InputPin, /*inout*/ float& Thickness, /*inout*/ FLinearColor& WireColor);
 
-	virtual bool IsConnectionCulled( const FArrangedWidget& StartLink, const FArrangedWidget& EndLink ) const;
+	UE_API virtual bool IsConnectionCulled( const FArrangedWidget& StartLink, const FArrangedWidget& EndLink ) const;
 
-	virtual TSharedPtr<IToolTip> GetConnectionToolTip(const SGraphPanel& GraphPanel, const FGraphSplineOverlapResult& OverlapData) const;
+	UE_API virtual TSharedPtr<IToolTip> GetConnectionToolTip(const SGraphPanel& GraphPanel, const FGraphSplineOverlapResult& OverlapData) const;
 	
 protected:
 	// Helper function used by Draw(). Called before DrawPinGeometries to populate PinToPinWidgetMap
-	virtual void BuildPinToPinWidgetMap(TMap<TSharedRef<SWidget>, FArrangedWidget>& InPinGeometries);
+	UE_API virtual void BuildPinToPinWidgetMap(TMap<TSharedRef<SWidget>, FArrangedWidget>& InPinGeometries);
 
 	// Helper function used by Draw(). Iterates over the pin geometries, drawing connections between them
-	virtual void DrawPinGeometries(TMap<TSharedRef<SWidget>, FArrangedWidget>& InPinGeometries, FArrangedChildren& ArrangedNodes);
+	UE_API virtual void DrawPinGeometries(TMap<TSharedRef<SWidget>, FArrangedWidget>& InPinGeometries, FArrangedChildren& ArrangedNodes);
 
 private:
 	// Note - these helpers are temporary, to be replaced entirely by the public API's 2f/float functions once the 2D/double deprecated functions can be removed
-	void DrawSplineWithArrow_DeprecationHelper(const UE::Slate::FDeprecateVector2DParameter& StartPoint, const UE::Slate::FDeprecateVector2DParameter& EndPoint, const FConnectionParams& Params);
-	UE::Slate::FDeprecateVector2DResult ComputeSplineTangent_DeprecationHelper(const UE::Slate::FDeprecateVector2DParameter& Start, const UE::Slate::FDeprecateVector2DParameter& End) const;
-	void DrawConnection_DeprecationHelper(int32 LayerId, const UE::Slate::FDeprecateVector2DParameter& Start, const UE::Slate::FDeprecateVector2DParameter& End, const FConnectionParams& Params);
-	void DrawPreviewConnector_DeprecationHelper(const FGeometry& PinGeometry, const UE::Slate::FDeprecateVector2DParameter& StartPoint, const UE::Slate::FDeprecateVector2DParameter& EndPoint, UEdGraphPin* Pin);
+	UE_API void DrawSplineWithArrow_DeprecationHelper(const UE::Slate::FDeprecateVector2DParameter& StartPoint, const UE::Slate::FDeprecateVector2DParameter& EndPoint, const FConnectionParams& Params);
+	UE_API UE::Slate::FDeprecateVector2DResult ComputeSplineTangent_DeprecationHelper(const UE::Slate::FDeprecateVector2DParameter& Start, const UE::Slate::FDeprecateVector2DParameter& End) const;
+	UE_API void DrawConnection_DeprecationHelper(int32 LayerId, const UE::Slate::FDeprecateVector2DParameter& Start, const UE::Slate::FDeprecateVector2DParameter& End, const FConnectionParams& Params);
+	UE_API void DrawPreviewConnector_DeprecationHelper(const FGeometry& PinGeometry, const UE::Slate::FDeprecateVector2DParameter& StartPoint, const UE::Slate::FDeprecateVector2DParameter& EndPoint, UEdGraphPin* Pin);
 };
+
+#undef UE_API
diff --git a/Engine/Source/Editor/GraphEditor/Public/EdGraphHandleTypes.h b/Engine/Source/Editor/GraphEditor/Public/EdGraphHandleTypes.h
--- a/Engine/Source/Editor/GraphEditor/Public/EdGraphHandleTypes.h
+++ b/Engine/Source/Editor/GraphEditor/Public/EdGraphHandleTypes.h
@@ -8,7 +8,7 @@
 #include "EdGraph/EdGraphPin.h"
 #include "UObject/SoftObjectPtr.h"
 
-struct GRAPHEDITOR_API FEdGraphNodeHandle
+struct FEdGraphNodeHandle
 {
 public:
 
@@ -51,7 +51,7 @@ private:
 	FName NodeName;
 };
 
-struct GRAPHEDITOR_API FEdGraphPinHandle : FEdGraphNodeHandle
+struct FEdGraphPinHandle : FEdGraphNodeHandle
 {
 public:
 
diff --git a/Engine/Source/Editor/GraphEditor/Public/FindInGraph.h b/Engine/Source/Editor/GraphEditor/Public/FindInGraph.h
--- a/Engine/Source/Editor/GraphEditor/Public/FindInGraph.h
+++ b/Engine/Source/Editor/GraphEditor/Public/FindInGraph.h
@@ -14,11 +14,13 @@
 #include "Widgets/Views/STableRow.h"
 #include "Widgets/Views/STreeView.h"
 
+#define UE_API GRAPHEDITOR_API
+
 class FAssetEditorToolkit;
 class SSearchBox;
 
 /** Item that matched the search results */
-class GRAPHEDITOR_API FFindInGraphResult
+class FFindInGraphResult
 {
 public:
 	virtual ~FFindInGraphResult() = default;
@@ -34,25 +36,25 @@ public:
 	};
 
 	/** Create a root (or only text) result */
-	FFindInGraphResult(const FCreateParams& InCreateParams);
+	UE_API FFindInGraphResult(const FCreateParams& InCreateParams);
 
 	/** By default NoOp, implement this to make your result jump to a node or node owning pin on click*/
-	virtual void JumpToNode(TWeakPtr<FAssetEditorToolkit> AssetEditorToolkit, const UEdGraphNode* InNode) const;
+	UE_API virtual void JumpToNode(TWeakPtr<FAssetEditorToolkit> AssetEditorToolkit, const UEdGraphNode* InNode) const;
 
 	/** Called when user clicks on the search item */
-	virtual FReply OnClick(TWeakPtr<FAssetEditorToolkit> AssetEditorToolkit);
+	UE_API virtual FReply OnClick(TWeakPtr<FAssetEditorToolkit> AssetEditorToolkit);
 
 	/* Get Category for this search result */
-	virtual FText GetCategory() const;
+	UE_API virtual FText GetCategory() const;
 
 	/** Create an icon to represent the result */
-	virtual TSharedRef<SWidget>	CreateIcon() const;
+	UE_API virtual TSharedRef<SWidget>	CreateIcon() const;
 
 	/** Gets the comment on this node if any */
-	virtual FString GetCommentText() const;
+	UE_API virtual FString GetCommentText() const;
 
 	/** Gets the value of the pin if any */
-	virtual FString GetValueText() const;
+	UE_API virtual FString GetValueText() const;
 
 	/** Any children listed under this category */
 	TArray<TSharedPtr<FFindInGraphResult>> Children;
@@ -81,62 +83,62 @@ public:
 
 
 /** Widget for searching for items that are part of a UEdGraph */
-class GRAPHEDITOR_API SFindInGraph : public SCompoundWidget
+class SFindInGraph : public SCompoundWidget
 {
 public:
 	SLATE_BEGIN_ARGS(SFindInGraph){}
 	SLATE_END_ARGS()
 
-	void Construct(const FArguments& InArgs, TSharedPtr<FAssetEditorToolkit> InAssetEditorToolkit);
+	UE_API void Construct(const FArguments& InArgs, TSharedPtr<FAssetEditorToolkit> InAssetEditorToolkit);
 
 	/** Focuses this widget's search box */
-	virtual void FocusForUse();
+	UE_API virtual void FocusForUse();
 
 protected:
 	typedef TSharedPtr<FFindInGraphResult> FSearchResult;
 	typedef STreeView<FSearchResult> STreeViewType;
 
 	/** Override this to create the search result of your type */
-	virtual TSharedPtr<FFindInGraphResult> MakeSearchResult(const FFindInGraphResult::FCreateParams& InParams);
+	UE_API virtual TSharedPtr<FFindInGraphResult> MakeSearchResult(const FFindInGraphResult::FCreateParams& InParams);
 
 	/** Called when user changes the text they are searching for */
-	virtual void OnSearchTextChanged(const FText& Text);
+	UE_API virtual void OnSearchTextChanged(const FText& Text);
 
 	/** Called when user changes commits text to the search box */
-	virtual void OnSearchTextCommitted(const FText& Text, ETextCommit::Type CommitType);
+	UE_API virtual void OnSearchTextCommitted(const FText& Text, ETextCommit::Type CommitType);
 
 	/** Get the children of a row */
-	virtual void OnGetChildren(FSearchResult InItem, TArray<FSearchResult>& OutChildren);
+	UE_API virtual void OnGetChildren(FSearchResult InItem, TArray<FSearchResult>& OutChildren);
 
 	/** Called when user clicks on a new result */
-	virtual void OnTreeSelectionChanged(FSearchResult Item, ESelectInfo::Type SelectInfo);
+	UE_API virtual void OnTreeSelectionChanged(FSearchResult Item, ESelectInfo::Type SelectInfo);
 
 	/** Called when user double clicks on a new result */
-	virtual void OnTreeSelectionDoubleClick(FSearchResult Item);
+	UE_API virtual void OnTreeSelectionDoubleClick(FSearchResult Item);
 
 	/** Called when a new row is being generated */
-	virtual TSharedRef<ITableRow> OnGenerateRow(FSearchResult InItem, const TSharedRef<STableViewBase>& OwnerTable);
+	UE_API virtual TSharedRef<ITableRow> OnGenerateRow(FSearchResult InItem, const TSharedRef<STableViewBase>& OwnerTable);
 
 	/** Begins the search based on the SearchValue */
-	virtual void InitiateSearch();
+	UE_API virtual void InitiateSearch();
 
 	/** Get a graph corresponding to this editor, implement if you want to want to use the default graph node search */
-	virtual const UEdGraph* GetGraph();
+	UE_API virtual const UEdGraph* GetGraph();
 
 	/** Find any results that contain all of the tokens */
-	virtual void MatchTokens(const TArray<FString>& Tokens);
+	UE_API virtual void MatchTokens(const TArray<FString>& Tokens);
 
 	/** 
 	 * Implement this if you have node-specific search behavior. IE: Casting into child types to compare tokens. 
 	 * @return true if all tokens are valid for this node, false otherwisee
 	 */
-	virtual bool MatchTokensInNode(const UEdGraphNode* Node, const TArray<FString>& Tokens);
+	UE_API virtual bool MatchTokensInNode(const UEdGraphNode* Node, const TArray<FString>& Tokens);
 
 	/** Find any results that contain all of the tokens in provided graph and subgraphs */
-	virtual void MatchTokensInGraph(const UEdGraph* Graph, const TArray<FString>& Tokens);
+	UE_API virtual void MatchTokensInGraph(const UEdGraph* Graph, const TArray<FString>& Tokens);
 
 	/** Determines if a string matches the search tokens */
-	static bool StringMatchesSearchTokens(const TArray<FString>& Tokens, const FString& ComparisonString);
+	static UE_API bool StringMatchesSearchTokens(const TArray<FString>& Tokens, const FString& ComparisonString);
 
 protected:
 	/** Pointer back to the Material editor that owns us */
@@ -160,3 +162,5 @@ protected:
 	/** The string to search for */
 	FString	SearchValue;
 };
+
+#undef UE_API
diff --git a/Engine/Source/Editor/GraphEditor/Public/GraphActionNode.h b/Engine/Source/Editor/GraphEditor/Public/GraphActionNode.h
--- a/Engine/Source/Editor/GraphEditor/Public/GraphActionNode.h
+++ b/Engine/Source/Editor/GraphEditor/Public/GraphActionNode.h
@@ -13,10 +13,12 @@
 #include "Templates/SharedPointer.h"
 #include "Templates/TypeHash.h"
 
+#define UE_API GRAPHEDITOR_API
+
 template <typename ItemType> class STreeView;
 
 // Utility class for building menus of graph actions
-struct GRAPHEDITOR_API FGraphActionNode : TSharedFromThis<FGraphActionNode>
+struct FGraphActionNode : TSharedFromThis<FGraphActionNode>
 {
 public:
 	// We need to declare our copy constructors so that we can disable
@@ -55,7 +57,7 @@ public:
 	 *
 	 * @return A newly allocated root node (should not be displayed in the tree view).
 	 */
-	static TSharedPtr<FGraphActionNode> NewRootNode();
+	static UE_API TSharedPtr<FGraphActionNode> NewRootNode();
 
 	/**
 	 * Inserts a new action node (and any accompanying category nodes) based off
@@ -67,11 +69,11 @@ public:
 	 * @param  Action	An action that you want the node to execute when picked.
 	 * @return The new action node.
 	 */
-	TSharedPtr<FGraphActionNode> AddChild(const TSharedPtr<FEdGraphSchemaAction>& Action);
+	UE_API TSharedPtr<FGraphActionNode> AddChild(const TSharedPtr<FEdGraphSchemaAction>& Action);
 	PRAGMA_DISABLE_DEPRECATION_WARNINGS
 	UE_DEPRECATED(5.5, "FGraphActionListBuilderBase::ActionGroup has been deprecated, use TSharedPtr<FEdGraphSchemaAction> directly")
-	TSharedPtr<FGraphActionNode> AddChild(FGraphActionListBuilderBase::ActionGroup const& ActionSet);
-	PRAGMA_ENABLE_DEPRECATION_WARNINGS
+	UE_API TSharedPtr<FGraphActionNode> AddChild(FGraphActionListBuilderBase::ActionGroup const& ActionSet);
+	UE_API PRAGMA_ENABLE_DEPRECATION_WARNINGS
 
 	/**
 	 * Inserts a new action node (and any required category nodes) based off
@@ -83,8 +85,8 @@ public:
 	TSharedPtr<FGraphActionNode> AddChildAlphabetical(const TSharedPtr<FEdGraphSchemaAction>& Action);
 	PRAGMA_DISABLE_DEPRECATION_WARNINGS
 	UE_DEPRECATED(5.5, "FGraphActionListBuilderBase::ActionGroup has been deprecated, use TSharedPtr<FEdGraphSchemaAction> directly")
-	TSharedPtr<FGraphActionNode> AddChildAlphabetical(FGraphActionListBuilderBase::ActionGroup const& ActionSet);
-	PRAGMA_ENABLE_DEPRECATION_WARNINGS
+	UE_API TSharedPtr<FGraphActionNode> AddChildAlphabetical(FGraphActionListBuilderBase::ActionGroup const& ActionSet);
+	UE_API PRAGMA_ENABLE_DEPRECATION_WARNINGS
 
 	TSharedPtr<FGraphActionNode> AddSection(int32 Grouping, int32 InSectionID);
 
@@ -95,7 +97,7 @@ public:
 	 * @param  bAlphabetically	Determines if we sort alphabetically on top of section/group/type.
 	 * @param  bRecursive		Determines if we should sort all decendent nodes' children ass well.
 	 */
-	void SortChildren(bool bAlphabetically = true, bool bRecursive = true);
+	UE_API void SortChildren(bool bAlphabetically = true, bool bRecursive = true);
 
 	/**
 	 * Returns a WeakPtr to the Parent Node
@@ -107,14 +109,14 @@ public:
 	 * 
 	 * @param  OutNodeArray	The array to fill out with decendent nodes.
 	 */
-	void GetAllNodes(TArray< TSharedPtr<FGraphActionNode> >& OutNodeArray) const;
+	UE_API void GetAllNodes(TArray< TSharedPtr<FGraphActionNode> >& OutNodeArray) const;
 
 	/**
 	 * Recursively collects all child/grandchild/decendent action nodes.
 	 * 
 	 * @param  OutNodeArray	The array to fill out with decendent nodes.
 	 */
-	void GetAllActionNodes(TArray<TSharedPtr<FGraphActionNode>>& OutNodeArray) const;
+	UE_API void GetAllActionNodes(TArray<TSharedPtr<FGraphActionNode>>& OutNodeArray) const;
 
 	/**
 	 * Recursively collects all decendent action/separator nodes (leaves out 
@@ -122,10 +124,10 @@ public:
 	 * 
 	 * @param  OutLeafArray	The array to fill out with decendent leaf nodes.
 	 */
-	void GetLeafNodes(TArray< TSharedPtr<FGraphActionNode> >& OutLeafArray) const;
+	UE_API void GetLeafNodes(TArray< TSharedPtr<FGraphActionNode> >& OutLeafArray) const;
 
 	/** Returns the number of leaf nodes */
-	int32 GetTotalLeafNodes() const;
+	UE_API int32 GetTotalLeafNodes() const;
 
 	/**
 	 * Takes the tree view and expands its elements for each child.
@@ -133,13 +135,13 @@ public:
 	 * @param  TreeView		The tree responsible for visualizing this node hierarchy.
 	 * @param  bRecursive	Determines if you want children/decendents to expand their children as well. 
 	 */
-	void ExpandAllChildren(TSharedPtr< STreeView< TSharedPtr<FGraphActionNode> > > TreeView, bool bRecursive = true);
+	UE_API void ExpandAllChildren(TSharedPtr< STreeView< TSharedPtr<FGraphActionNode> > > TreeView, bool bRecursive = true);
 
 	/**
 	 * Clears all children (not recursively... the TSharedPtrs should clean up 
 	 * appropriately).
 	 */
-	void ClearChildren();
+	UE_API void ClearChildren();
 
 	/**
 	 * Query to determine this node's type (there are five distinguishable node
@@ -147,11 +149,11 @@ public:
 	 *
 	 * @return True if this is the type your queried about, otherwise false.
 	 */
-	bool IsRootNode() const;
-	bool IsSectionHeadingNode() const;
-	bool IsCategoryNode() const;
-	bool IsActionNode() const;
-	bool IsGroupDividerNode() const;
+	UE_API bool IsRootNode() const;
+	UE_API bool IsSectionHeadingNode() const;
+	UE_API bool IsCategoryNode() const;
+	UE_API bool IsActionNode() const;
+	UE_API bool IsGroupDividerNode() const;
 
 	/**
 	 * Determines if this node is a menu separator of some kind (either a
@@ -159,14 +161,14 @@ public:
 	 *
 	 * @return True if this is a menu divider, otherwise false.
 	 */
-	bool IsSeparator() const;
+	UE_API bool IsSeparator() const;
 	/**
 	 * Retrieves this node's display name (for category and action nodes). The
 	 * text string will be empty for separator and root nodes.
 	 *
 	 * @return The name to present this node with in the tree view (will be an empty text string if this is a separator node)
 	 */
-	FText const& GetDisplayName() const;
+	UE_API FText const& GetDisplayName() const;
 
 	/**
 	 * Walks the node chain backwards, constructing a category path (delimited
@@ -175,14 +177,14 @@ public:
 	 *
 	 * @return A category path string, denoting the category hierarchy up to this node.
 	 */
-	FText GetCategoryPath() const;
+	UE_API FText GetCategoryPath() const;
 
 	/**
 	 * Checks to see if this node contains at least one valid action.
 	 *
 	 * @return True is the Actions array contains a valid entry, otherwise false.
 	 */
-	bool HasValidAction() const;
+	UE_API bool HasValidAction() const;
 
 	/**
 	 * Looks through this node's Actions array, and returns the first valid
@@ -190,7 +192,7 @@ public:
 	 *
 	 * @return This node's first valid action (will be an empty pointer if this is not an action node).
 	 */
-	TSharedPtr<FEdGraphSchemaAction> GetPrimaryAction() const;
+	UE_API TSharedPtr<FEdGraphSchemaAction> GetPrimaryAction() const;
 
 	/**
 	 * Accessor to the node's RenameRequestEvent (for binding purposes). Do not
@@ -207,7 +209,7 @@ public:
 	 *
 	 * @return True if the broadcast went through, false if the "pending rename request" flag was set.
 	 */
-	bool BroadcastRenameRequest();
+	UE_API bool BroadcastRenameRequest();
 
 	/**
 	 * Sometimes a call to BroadcastRenameRequest() is made before the
@@ -217,10 +219,10 @@ public:
 	 *
 	 * @return True if a call to BroadcastRenameRequest() was made without a valid RenameRequestEvent.
 	 */
-	bool IsRenameRequestPending() const;
+	UE_API bool IsRenameRequestPending() const;
 
 	/** Returns the 'linearized' index of the node, including category nodes, useful for getting displayed position */
-	int32 GetLinearizedIndex(TSharedPtr<FGraphActionNode> Node) const;
+	UE_API int32 GetLinearizedIndex(TSharedPtr<FGraphActionNode> Node) const;
 private:
 	/**
 	 *
@@ -228,7 +230,7 @@ private:
 	 * @param  Grouping
 	 * @param  SectionID
 	 */
-	FGraphActionNode(int32 Grouping, int32 SectionID);
+	UE_API FGraphActionNode(int32 Grouping, int32 SectionID);
 
 	/**
 	 * Constructor for action nodes. Private so that users go through AddChild().
@@ -237,7 +239,7 @@ private:
 	 * @param  Grouping
 	 * @param  SectionID
 	 */
-	FGraphActionNode(const TSharedPtr<FEdGraphSchemaAction>& InAction, int32 InGrouping, int32 InSectionID);
+	UE_API FGraphActionNode(const TSharedPtr<FEdGraphSchemaAction>& InAction, int32 InGrouping, int32 InSectionID);
 
 	/**
 	 *
@@ -247,7 +249,7 @@ private:
 	 * @param  SectionID
 	 * @return
 	 */
-	static TSharedPtr<FGraphActionNode> NewSectionHeadingNode(TWeakPtr<FGraphActionNode> Parent, int32 Grouping, int32 SectionID);
+	static UE_API TSharedPtr<FGraphActionNode> NewSectionHeadingNode(TWeakPtr<FGraphActionNode> Parent, int32 Grouping, int32 SectionID);
 
 	/**
 	 *
@@ -257,7 +259,7 @@ private:
 	 * @param  SectionID
 	 * @return
 	 */
-	static TSharedPtr<FGraphActionNode> NewCategoryNode(FString const& Category, int32 Grouping, int32 SectionID);
+	static UE_API TSharedPtr<FGraphActionNode> NewCategoryNode(FString const& Category, int32 Grouping, int32 SectionID);
 
 	/**
 	 *
@@ -265,7 +267,7 @@ private:
 	 * @param  ActionNode
 	 * @return
 	 */
-	static TSharedPtr<FGraphActionNode> NewActionNode(TSharedPtr<FEdGraphSchemaAction> const& ActionNode);
+	static UE_API TSharedPtr<FGraphActionNode> NewActionNode(TSharedPtr<FEdGraphSchemaAction> const& ActionNode);
 
 	/**
 	 *
@@ -274,7 +276,7 @@ private:
 	 * @param  Grouping
 	 * @return
 	 */
-	static TSharedPtr<FGraphActionNode> NewGroupDividerNode(TWeakPtr<FGraphActionNode> Parent, int32 Grouping);
+	static UE_API TSharedPtr<FGraphActionNode> NewGroupDividerNode(TWeakPtr<FGraphActionNode> Parent, int32 Grouping);
 
 	/**
 	 * Iterates the CategoryStack, adding category-nodes as needed. The
@@ -284,7 +286,7 @@ private:
 	 * @param  Idx				Current point in the category stack that we have iterated to
 	 * @param  NodeToAdd		The node you want inserted.
 	 */
-	void AddChildRecursively(const TArray<FString>& CategoryStack, int32 Idx, TSharedPtr<FGraphActionNode> NodeToAdd);
+	UE_API void AddChildRecursively(const TArray<FString>& CategoryStack, int32 Idx, TSharedPtr<FGraphActionNode> NodeToAdd);
 
 	/**
 	 * Looks through this node's children to see if a there already exists a 
@@ -294,7 +296,7 @@ private:
 	 * @param   NodeToAdd	The node that we'll be adding to this child.
 	 * @return  A child node matching the supplied parameters (will be empty if no match was found).
 	 */
-	TSharedPtr<FGraphActionNode> FindMatchingParent(FString const& ParentName, TSharedPtr<FGraphActionNode> NodeToAdd);
+	UE_API TSharedPtr<FGraphActionNode> FindMatchingParent(FString const& ParentName, TSharedPtr<FGraphActionNode> NodeToAdd);
 
 	/**
 	 * Adds the specified node directly to this node's Children array. Will
@@ -303,13 +305,13 @@ private:
 	 *
 	 * @param  NodeToAdd The node you want inserted.
 	 */
-	void InsertChild(TSharedPtr<FGraphActionNode> NodeToAdd);
+	UE_API void InsertChild(TSharedPtr<FGraphActionNode> NodeToAdd);
 
-	void AddChildGrouping(TSharedPtr<FGraphActionNode> ActionNode, TWeakPtr<FGraphActionNode> Parent, bool bInsertAlphabetically);
-	void InsertChildAlphabetical(TSharedPtr<FGraphActionNode> NodeToAdd);
+	UE_API void AddChildGrouping(TSharedPtr<FGraphActionNode> ActionNode, TWeakPtr<FGraphActionNode> Parent, bool bInsertAlphabetically);
+	UE_API void InsertChildAlphabetical(TSharedPtr<FGraphActionNode> NodeToAdd);
 
 	/** Recursive implementation helper for GetLinearizedIndex */
-	int32 GetLinearizedIndex(TSharedPtr<FGraphActionNode> Node, int32& Iter) const;
+	UE_API int32 GetLinearizedIndex(TSharedPtr<FGraphActionNode> Node, int32& Iter) const;
 private:
 	/** The category or action name (depends on what type of node this is) */
 	FText DisplayText;
@@ -332,3 +334,5 @@ private:
 	/** Root entry only, counts the total leaf entries in this tree */
 	int32 TotalLeafs;
 };
+
+#undef UE_API
diff --git a/Engine/Source/Editor/GraphEditor/Public/GraphDiffControl.h b/Engine/Source/Editor/GraphEditor/Public/GraphDiffControl.h
--- a/Engine/Source/Editor/GraphEditor/Public/GraphDiffControl.h
+++ b/Engine/Source/Editor/GraphEditor/Public/GraphDiffControl.h
@@ -7,13 +7,15 @@
 #include "DiffResults.h"
 #include "Internationalization/Text.h"
 
+#define UE_API GRAPHEDITOR_API
+
 class UEdGraph;
 class UEdGraphNode;
 struct FDiffResults;
 struct FDiffSingleResult;
 
 /** Used to find differences between revisions of a graph. */
-class GRAPHEDITOR_API FGraphDiffControl
+class FGraphDiffControl
 {
 public:
 	/** What kind of diff are we performing */
@@ -46,7 +48,7 @@ public:
 	};
 
 	/** A struct holding the context data for a node diff */
-	struct GRAPHEDITOR_API FNodeDiffContext
+	struct FNodeDiffContext
 	{
 		FNodeDiffContext()
 			: DiffMode(EDiffMode::Additive)
@@ -70,7 +72,7 @@ public:
 	};
 
 	/** A struct to represent a found pair of nodes that match each other (for comparisons sake) */
-	struct GRAPHEDITOR_API FNodeMatch
+	struct FNodeMatch
 	{
 		FNodeMatch()
 			: NewNode(nullptr)
@@ -88,14 +90,14 @@ public:
 		 * @param  DiffsArrayOut	If supplied, this will be filled out with all the differences that were found.
 		 * @return True if there were differences found, false if the two nodes are identical.
 		 */
-		bool Diff(const FNodeDiffContext& DiffContext, TArray<FDiffSingleResult>* DiffsResultsOut = nullptr) const;
-		bool Diff(const FNodeDiffContext& DiffContext, FDiffResults& DiffsOut) const;
+		UE_API bool Diff(const FNodeDiffContext& DiffContext, TArray<FDiffSingleResult>* DiffsResultsOut = nullptr) const;
+		UE_API bool Diff(const FNodeDiffContext& DiffContext, FDiffResults& DiffsOut) const;
 
 		/**
 		 * Checks to see if this is a valid match.
 		 * @return False if NewNode or OldNode is nullptr, true if both are valid.
 		 */
-		bool IsValid() const;
+		UE_API bool IsValid() const;
 	};
 
 	/**
@@ -109,7 +111,7 @@ public:
 	 * @param  PriorMatches	Previous made matches to exclude from our search.
 	 * @return A pair of nodes (including the supplied one) that best match each other (one may be nullptr if no match was found).
 	 */
-	static FNodeMatch FindNodeMatch(UEdGraph* OldGraph, UEdGraphNode* NewNode, TArray<FNodeMatch> const& PriorMatches);
+	static UE_API FNodeMatch FindNodeMatch(UEdGraph* OldGraph, UEdGraphNode* NewNode, TArray<FNodeMatch> const& PriorMatches);
 
 	/** 
 	 * Returns true if the two Nodes are the same 
@@ -118,7 +120,7 @@ public:
 	 * @param  Node2		Second node
 	 * @param  Exclusions	Previous made matches to exclude from our search.
 	 */
-	static bool IsNodeMatch(UEdGraphNode* Node1, UEdGraphNode* Node2, bool bExactOnly = false, TArray<FGraphDiffControl::FNodeMatch> const* Exclusions = nullptr);
+	static UE_API bool IsNodeMatch(UEdGraphNode* Node1, UEdGraphNode* Node2, bool bExactOnly = false, TArray<FGraphDiffControl::FNodeMatch> const* Exclusions = nullptr);
 
 	/**
 	 * Looks for node differences between the two supplied graphs. Diffs will be
@@ -129,8 +131,10 @@ public:
 	 * @param  DiffsOut		All the differences that were found between the two.
 	 * @return True if any differences were found, false if both graphs are identical.
 	 */
-	static bool DiffGraphs(UEdGraph* const OldGraph, UEdGraph* const NewGraph, TArray<FDiffSingleResult>& DiffsOut);
+	static UE_API bool DiffGraphs(UEdGraph* const OldGraph, UEdGraph* const NewGraph, TArray<FDiffSingleResult>& DiffsOut);
 
 	/** Computes an object path for a specific graph, relative to the blueprint/asset root */
-	static FString GetGraphPath(UEdGraph* Graph);
+	static UE_API FString GetGraphPath(UEdGraph* Graph);
 };
+
+#undef UE_API
diff --git a/Engine/Source/Editor/GraphEditor/Public/GraphEditorDragDropAction.h b/Engine/Source/Editor/GraphEditor/Public/GraphEditorDragDropAction.h
--- a/Engine/Source/Editor/GraphEditor/Public/GraphEditorDragDropAction.h
+++ b/Engine/Source/Editor/GraphEditor/Public/GraphEditorDragDropAction.h
@@ -15,6 +15,8 @@
 #include "Templates/SharedPointer.h"
 #include "UObject/WeakObjectPtrTemplates.h"
 
+#define UE_API GRAPHEDITOR_API
+
 class SGraphNode;
 class SGraphPanel;
 class SWidget;
@@ -26,52 +28,52 @@ struct FEdGraphSchemaAction;
 struct FSlateBrush;
 
 // Base class for drag-drop actions that pass into the graph editor and perform an action when dropped
-class GRAPHEDITOR_API FGraphEditorDragDropAction : public FDragDropOperation
+class FGraphEditorDragDropAction : public FDragDropOperation
 {
 public:
 	DRAG_DROP_OPERATOR_TYPE(FGraphEditorDragDropAction, FDragDropOperation)
 
-	FGraphEditorDragDropAction();
+	UE_API FGraphEditorDragDropAction();
 
-	void SetHoveredPin(UEdGraphPin* InPin);
-	void SetHoveredNode(const TSharedPtr<SGraphNode>& InNode);
-	void SetHoveredNode(UEdGraphNode* InNode);
-	void SetHoveredGraph(const TSharedPtr<SGraphPanel>& InGraph);
-	void SetHoveredCategoryName(const FText& InHoverCategoryName);
-	void SetHoveredAction(TSharedPtr<struct FEdGraphSchemaAction> Action);
+	UE_API void SetHoveredPin(UEdGraphPin* InPin);
+	UE_API void SetHoveredNode(const TSharedPtr<SGraphNode>& InNode);
+	UE_API void SetHoveredNode(UEdGraphNode* InNode);
+	UE_API void SetHoveredGraph(const TSharedPtr<SGraphPanel>& InGraph);
+	UE_API void SetHoveredCategoryName(const FText& InHoverCategoryName);
+	UE_API void SetHoveredAction(TSharedPtr<struct FEdGraphSchemaAction> Action);
 	void SetDropTargetValid( bool bValid ) { bDropTargetValid = bValid; }
 
 	// Interface to override
 	virtual void HoverTargetChanged() {}
 	UE_DEPRECATED(5.6, "Use the version of the function accepting FVector2f; this Slate API no longer interfaces directly with double-precision scalars and vectors.")
-	virtual FReply DroppedOnPin(FVector2D ScreenPosition, FVector2D GraphPosition) UE_SLATE_DEPRECATED_VECTOR_VIRTUAL_FUNCTION;
-	virtual FReply DroppedOnPin(const FVector2f& ScreenPosition, const FVector2f& GraphPosition);
+	UE_API virtual FReply DroppedOnPin(FVector2D ScreenPosition, FVector2D GraphPosition) UE_SLATE_DEPRECATED_VECTOR_VIRTUAL_FUNCTION;
+	UE_API virtual FReply DroppedOnPin(const FVector2f& ScreenPosition, const FVector2f& GraphPosition);
 	UE_DEPRECATED(5.6, "Use the version of the function accepting FVector2f; this Slate API no longer interfaces directly with double-precision scalars and vectors.")
-	virtual FReply DroppedOnNode(FVector2D ScreenPosition, FVector2D GraphPosition) UE_SLATE_DEPRECATED_VECTOR_VIRTUAL_FUNCTION;
-	virtual FReply DroppedOnNode(const FVector2f& ScreenPosition, const FVector2f& GraphPosition);
+	UE_API virtual FReply DroppedOnNode(FVector2D ScreenPosition, FVector2D GraphPosition) UE_SLATE_DEPRECATED_VECTOR_VIRTUAL_FUNCTION;
+	UE_API virtual FReply DroppedOnNode(const FVector2f& ScreenPosition, const FVector2f& GraphPosition);
 	UE_DEPRECATED(5.6, "Use the version of the function accepting FVector2f; this Slate API no longer interfaces directly with double-precision scalars and vectors.")
-	virtual FReply DroppedOnPanel(const TSharedRef< SWidget >& Panel, FVector2D ScreenPosition, FVector2D GraphPosition, UEdGraph& Graph) UE_SLATE_DEPRECATED_VECTOR_VIRTUAL_FUNCTION;
-	virtual FReply DroppedOnPanel(const TSharedRef< SWidget >& Panel, const FVector2f& ScreenPosition, const FVector2f& GraphPosition, UEdGraph& Graph);
+	UE_API virtual FReply DroppedOnPanel(const TSharedRef< SWidget >& Panel, FVector2D ScreenPosition, FVector2D GraphPosition, UEdGraph& Graph) UE_SLATE_DEPRECATED_VECTOR_VIRTUAL_FUNCTION;
+	UE_API virtual FReply DroppedOnPanel(const TSharedRef< SWidget >& Panel, const FVector2f& ScreenPosition, const FVector2f& GraphPosition, UEdGraph& Graph);
 	virtual FReply DroppedOnAction(TSharedRef<struct FEdGraphSchemaAction> Action) { return FReply::Unhandled(); }
 	virtual FReply DroppedOnCategory(FText Category) { return FReply::Unhandled(); }
 	// End of interface to override
 	
 	virtual bool IsSupportedBySchema(const UEdGraphSchema* Schema) const { return true; }
 
-	bool HasFeedbackMessage();
-	void SetFeedbackMessage(const TSharedPtr<SWidget>& Message);
-	void SetSimpleFeedbackMessage(const FSlateBrush* Icon, const FSlateColor& IconColor, const FText& Message, const FSlateBrush* SecondaryIcon = nullptr, const FSlateColor SecondaryColor = FSlateColor());
+	UE_API bool HasFeedbackMessage();
+	UE_API void SetFeedbackMessage(const TSharedPtr<SWidget>& Message);
+	UE_API void SetSimpleFeedbackMessage(const FSlateBrush* Icon, const FSlateColor& IconColor, const FText& Message, const FSlateBrush* SecondaryIcon = nullptr, const FSlateColor SecondaryColor = FSlateColor());
 
 protected:
-	UEdGraphPin* GetHoveredPin() const;
-	UEdGraphNode* GetHoveredNode() const;
-	UEdGraph* GetHoveredGraph() const;
+	UE_API UEdGraphPin* GetHoveredPin() const;
+	UE_API UEdGraphNode* GetHoveredNode() const;
+	UE_API UEdGraph* GetHoveredGraph() const;
 
 	/** Constructs the window and widget if applicable */
-	virtual void Construct() override;
+	UE_API virtual void Construct() override;
 
-	virtual EVisibility GetIconVisible() const;
-	virtual EVisibility GetErrorIconVisible() const;
+	UE_API virtual EVisibility GetIconVisible() const;
+	UE_API virtual EVisibility GetErrorIconVisible() const;
 
 private:
 	// The pin that the drag action is currently hovering over
@@ -96,16 +98,16 @@ protected:
 };
 
 // Drag-drop action where an FEdGraphSchemaAction should be performed when dropped
-class GRAPHEDITOR_API FGraphSchemaActionDragDropAction : public FGraphEditorDragDropAction
+class FGraphSchemaActionDragDropAction : public FGraphEditorDragDropAction
 {
 public:
 	
 	DRAG_DROP_OPERATOR_TYPE(FGraphSchemaActionDragDropAction, FGraphEditorDragDropAction)
 
 	// FGraphEditorDragDropAction interface
-	virtual void HoverTargetChanged() override;
-	virtual FReply DroppedOnPanel(const TSharedRef< SWidget >& Panel, const FVector2f& ScreenPosition, const FVector2f& GraphPosition, UEdGraph& Graph) override;
-	virtual FReply DroppedOnPin(const FVector2f& ScreenPosition, const FVector2f& GraphPosition) override;
+	UE_API virtual void HoverTargetChanged() override;
+	UE_API virtual FReply DroppedOnPanel(const TSharedRef< SWidget >& Panel, const FVector2f& ScreenPosition, const FVector2f& GraphPosition, UEdGraph& Graph) override;
+	UE_API virtual FReply DroppedOnPin(const FVector2f& ScreenPosition, const FVector2f& GraphPosition) override;
 	// End of FGraphEditorDragDropAction
 
 	static TSharedRef<FGraphSchemaActionDragDropAction> New(TSharedPtr<FEdGraphSchemaAction> InActionNode )
@@ -117,8 +119,10 @@ public:
 	}
 
 protected:
-	virtual void GetDefaultStatusSymbol(const FSlateBrush*& PrimaryBrushOut, FSlateColor& IconColorOut, FSlateBrush const*& SecondaryBrushOut, FSlateColor& SecondaryColorOut) const;
+	UE_API virtual void GetDefaultStatusSymbol(const FSlateBrush*& PrimaryBrushOut, FSlateColor& IconColorOut, FSlateBrush const*& SecondaryBrushOut, FSlateColor& SecondaryColorOut) const;
 
 	/** */
 	TSharedPtr<FEdGraphSchemaAction> SourceAction;
 };
+
+#undef UE_API
diff --git a/Engine/Source/Editor/GraphEditor/Public/GraphEditorSettings.h b/Engine/Source/Editor/GraphEditor/Public/GraphEditorSettings.h
--- a/Engine/Source/Editor/GraphEditor/Public/GraphEditorSettings.h
+++ b/Engine/Source/Editor/GraphEditor/Public/GraphEditorSettings.h
@@ -16,6 +16,8 @@
 
 #include "GraphEditorSettings.generated.h"
 
+#define UE_API GRAPHEDITOR_API
+
 UENUM()
 enum class EGraphPanningMouseButton : uint8
 {
@@ -29,8 +31,8 @@ struct FPropertyChangedEvent;
 /**
  * Implements settings for the graph editor.
  */
-UCLASS(config=EditorPerProjectUserSettings)
-class GRAPHEDITOR_API UGraphEditorSettings : public UObject
+UCLASS(MinimalAPI, config=EditorPerProjectUserSettings)
+class UGraphEditorSettings : public UObject
 {
 	GENERATED_UCLASS_BODY()
 
@@ -38,7 +40,7 @@ public:
 
 #if WITH_EDITOR
 	// UObject interface
-	virtual void PostEditChangeProperty(FPropertyChangedEvent& PropertyChangedEvent) override;
+	UE_API virtual void PostEditChangeProperty(FPropertyChangedEvent& PropertyChangedEvent) override;
 	// End of UObject interface
 #endif
 
@@ -81,16 +83,16 @@ public:
 	float PaddingTowardsNodeEdge;
 
 	/** Returns the padding to use for an input pin on a graph node */
-	FMargin GetInputPinPadding() const;
+	UE_API FMargin GetInputPinPadding() const;
 
 	/** Returns the padding to use for an output pin on a graph node */
-	FMargin GetOutputPinPadding() const;
+	UE_API FMargin GetOutputPinPadding() const;
 
 	/** Returns the padding to use for a node region that won't contain a pin */
-	FMargin GetNonPinNodeBodyPadding() const;
+	UE_API FMargin GetNonPinNodeBodyPadding() const;
 
 	/** Returns the size adjustment from node render geometry to shadow render geometry */
-	UE::Slate::FDeprecateVector2DResult GetShadowDeltaSize() const;
+	UE_API UE::Slate::FDeprecateVector2DResult GetShadowDeltaSize() const;
 
 public:
 
@@ -139,7 +141,7 @@ public:
 	FDeprecateSlateVector2D BackwardSplineTangentFromVerticalDelta;
 
 	/** Compute a tangent given the start and end points of the spline */
-	UE::Slate::FDeprecateVector2DResult ComputeSplineTangent(const UE::Slate::FDeprecateVector2DParameter& Start, const UE::Slate::FDeprecateVector2DParameter& End) const;
+	UE_API UE::Slate::FDeprecateVector2DResult ComputeSplineTangent(const UE::Slate::FDeprecateVector2DParameter& Start, const UE::Slate::FDeprecateVector2DParameter& End) const;
 
 public:
 
@@ -357,3 +359,5 @@ public:
 	UPROPERTY(EditAnywhere, config, Category = CommentNodes)
 	bool bShowCommentBubbleWhenZoomedOut;
 };
+
+#undef UE_API
diff --git a/Engine/Source/Editor/GraphEditor/Public/GraphSplineOverlapResult.h b/Engine/Source/Editor/GraphEditor/Public/GraphSplineOverlapResult.h
--- a/Engine/Source/Editor/GraphEditor/Public/GraphSplineOverlapResult.h
+++ b/Engine/Source/Editor/GraphEditor/Public/GraphSplineOverlapResult.h
@@ -5,10 +5,12 @@
 #include "CoreMinimal.h"
 #include "SGraphPin.h"
 
+#define UE_API GRAPHEDITOR_API
+
 /////////////////////////////////////////////////////
 // FGraphSplineOverlapResult
 
-struct GRAPHEDITOR_API FGraphSplineOverlapResult
+struct FGraphSplineOverlapResult
 {
 public:
 	FGraphSplineOverlapResult(bool InCloseToSpline = false)
@@ -42,7 +44,7 @@ public:
 		return DistanceSquared < FLT_MAX;
 	}
 
-	void ComputeBestPin();
+	UE_API void ComputeBestPin();
 
 	float GetDistanceSquared() const
 	{
@@ -77,8 +79,8 @@ public:
 	FGraphPinHandle GetPin1Handle() const { return Pin1Handle; }
 	FGraphPinHandle GetPin2Handle() const { return Pin2Handle; }
 
-	bool GetPins(const class SGraphPanel& InGraphPanel, UEdGraphPin*& OutPin1, UEdGraphPin*& OutPin2) const;
-	void GetPinWidgets(const class SGraphPanel& InGraphPanel, TSharedPtr<class SGraphPin>& OutPin1, TSharedPtr<class SGraphPin>& OutPin2) const;
+	UE_API bool GetPins(const class SGraphPanel& InGraphPanel, UEdGraphPin*& OutPin1, UEdGraphPin*& OutPin2) const;
+	UE_API void GetPinWidgets(const class SGraphPanel& InGraphPanel, TSharedPtr<class SGraphPin>& OutPin1, TSharedPtr<class SGraphPin>& OutPin2) const;
 
 protected:
 	FGraphPinHandle Pin1Handle;
@@ -91,3 +93,5 @@ protected:
 	float DistanceSquaredToPin2;
 	bool  bCloseToSpline;
 };
+
+#undef UE_API
