[HWRT] Enable HWRT by default for new projects.

#rnx

diff --git a/Engine/Source/Editor/GameProjectGeneration/Private/GameProjectUtils.cpp b/Engine/Source/Editor/GameProjectGeneration/Private/GameProjectUtils.cpp
--- a/Engine/Source/Editor/GameProjectGeneration/Private/GameProjectUtils.cpp
+++ b/Engine/Source/Editor/GameProjectGeneration/Private/GameProjectUtils.cpp
@@ -181,6 +181,12 @@ namespace
 			TEXT("True"),
 			false /* ShouldReplaceExistingValue */);
 
+		ConfigValues.Emplace(TEXT("DefaultEngine.ini"),
+			TEXT("/Script/Engine.RendererSettings"),
+			TEXT("r.RayTracing"),
+			TEXT("True"),
+			false /* ShouldReplaceExistingValue */);
+
 		ConfigValues.Emplace(TEXT("DefaultEngine.ini"),
 			TEXT("/Script/Engine.RendererSettings"),
 			TEXT("r.RayTracing.RayTracingProxies.ProjectEnabled"),
