The SNewClassDialog will now be called via SCustomDialog with Cancel and Confirm hotkeys enabled and proper widget focus on activation.

[RN]

diff --git a/Engine/Source/Editor/GameProjectGeneration/Private/GameProjectUtils.cpp b/Engine/Source/Editor/GameProjectGeneration/Private/GameProjectUtils.cpp
--- a/Engine/Source/Editor/GameProjectGeneration/Private/GameProjectUtils.cpp
+++ b/Engine/Source/Editor/GameProjectGeneration/Private/GameProjectUtils.cpp
@@ -77,6 +77,7 @@
 #include "Blueprint/BlueprintSupport.h"
 #include "Settings/ProjectPackagingSettings.h"
 #include "Async/ParallelFor.h"
+#include "Dialog/SCustomDialog.h"
 
 #if WITH_LIVE_CODING
 #include "ILiveCodingModule.h"
@@ -1120,12 +1121,10 @@ void GameProjectUtils::OpenAddToProjectDialog(const FAddToProjectConfig& Config,
 		WindowTitle = InDomain == EClassDomain::Native ? LOCTEXT("AddCodeWindowHeader_Native", "Add C++ Class") : LOCTEXT("AddCodeWindowHeader_Blueprint", "Add Blueprint Class");
 	}
 
-	TSharedRef<SWindow> AddCodeWindow =
-		SNew(SWindow)
+	TSharedRef<SCustomDialog> AddCodeWindow =
+		SNew(SCustomDialog)
 		.Title( WindowTitle )
-		.ClientSize( WindowSize )
-		.SizingRule( ESizingRule::UserSized )
-		.SupportsMinimize(false) .SupportsMaximize(false);
+		.ClientSize( WindowSize );
 
 	TSharedRef<SNewClassDialog> NewClassDialog =
 		SNew(SNewClassDialog)
@@ -1141,6 +1140,9 @@ void GameProjectUtils::OpenAddToProjectDialog(const FAddToProjectConfig& Config,
 		.SyncContentBrowserToNewClass(Config._SyncContentBrowserToNewClass);
 
 	AddCodeWindow->SetContent( NewClassDialog );
+	AddCodeWindow->SetOnCancelHotkeyPressed(FExecuteAction::CreateSP(NewClassDialog, &SNewClassDialog::CancelClicked));
+	AddCodeWindow->SetOnConfirmHotkeyPressed(FExecuteAction::CreateSP(NewClassDialog, &SNewClassDialog::OnConfirmHotkeyPressed));
+	AddCodeWindow->SetWidgetToFocusOnActivate(NewClassDialog->GetWidgetToFocusOnActivate());
 
 	TSharedPtr<SWindow> ParentWindow = Config._ParentWindow;
 	if (!ParentWindow.IsValid())
diff --git a/Engine/Source/Editor/GameProjectGeneration/Private/SNewClassDialog.cpp b/Engine/Source/Editor/GameProjectGeneration/Private/SNewClassDialog.cpp
--- a/Engine/Source/Editor/GameProjectGeneration/Private/SNewClassDialog.cpp
+++ b/Engine/Source/Editor/GameProjectGeneration/Private/SNewClassDialog.cpp
@@ -45,6 +45,7 @@
 #include "Widgets/Input/SSegmentedControl.h"
 #include "ClassIconFinder.h"
 #include "SWarningOrErrorBox.h"
+#include "Dialog/DialogCommands.h"
 
 #define LOCTEXT_NAMESPACE "GameProjectGeneration"
 
@@ -781,22 +782,9 @@ void SNewClassDialog::Construct( const FArguments& InArgs )
 }
 END_SLATE_FUNCTION_BUILD_OPTIMIZATION
 
-FReply SNewClassDialog::OnKeyDown(const FGeometry& MyGeometry, const FKeyEvent& InKeyEvent)
+TSharedPtr<SWidget> SNewClassDialog::GetWidgetToFocusOnActivate()
 {
-	if (InKeyEvent.GetKey() == EKeys::Escape)
-	{
-		// Pressing Escape returns as if the user clicked Cancel
-		CancelClicked();
-		return FReply::Handled();
-	}
-	else if (InKeyEvent.GetKey() == EKeys::Enter)
-	{
-		// Pressing Enter move to the next page like a double-click or the Next button
-		OnParentClassItemDoubleClicked(TSharedPtr<FParentClassItem>());
-		return FReply::Handled();
-	}
-
-	return FReply::Unhandled();
+	return MainWizard->GetWidgetToFocusOnActivate();
 }
 
 void SNewClassDialog::Tick( const FGeometry& AllottedGeometry, const double InCurrentTime, const float InDeltaTime )
@@ -1382,6 +1370,21 @@ TSharedRef<SWidget> SNewClassDialog::MakeWidgetForSelectedModuleCombo(TSharedPtr
 		.Text(FText::Format(LOCTEXT("ModuleComboEntry", "{ModuleName} ({ModuleType})"), Args));
 }
 
+void SNewClassDialog::OnConfirmHotkeyPressed()
+{
+	// Advance to the name page
+	constexpr int32 NamePageIdx = 1;
+	if (MainWizard->CanShowPage(NamePageIdx))
+	{
+		MainWizard->ShowPage(NamePageIdx);
+	}
+	// Otherwise finish
+	else if (CanFinish())
+	{
+		FinishClicked();
+	}
+}
+
 FSlateColor SNewClassDialog::GetClassLocationTextColor(GameProjectUtils::EClassLocation InLocation) const
 {
 	return (ClassLocation == InLocation) ? FSlateColor(FLinearColor(0, 0, 0)) : FSlateColor(FLinearColor(0.72f, 0.72f, 0.72f, 1.f));
diff --git a/Engine/Source/Editor/GameProjectGeneration/Private/SNewClassDialog.h b/Engine/Source/Editor/GameProjectGeneration/Private/SNewClassDialog.h
--- a/Engine/Source/Editor/GameProjectGeneration/Private/SNewClassDialog.h
+++ b/Engine/Source/Editor/GameProjectGeneration/Private/SNewClassDialog.h
@@ -90,11 +90,17 @@ public:
 	/** Constructs this widget with InArgs */
 	void Construct( const FArguments& InArgs );
 
-	/** Interpret Escape and Enter key press as Cancel or double-click/Next */
-	virtual FReply OnKeyDown( const FGeometry& MyGeometry, const FKeyEvent& InKeyEvent ) override;
+	/** Get the widget to focus on activation */
+	TSharedPtr<SWidget> GetWidgetToFocusOnActivate();
 
 	virtual void Tick( const FGeometry& AllottedGeometry, const double InCurrentTime, const float InDeltaTime ) override;
 
+	/** Handler for when cancel is clicked */
+	void CancelClicked();
+
+	/** Handler for when the Confirm hotkey is pressed */
+	void OnConfirmHotkeyPressed();
+
 private:
 
 	/** Creates a row in the parent class list */
@@ -115,7 +121,7 @@ private:
 	/** Whether the document link anchor should be visible */
 	EVisibility GetDocLinkVisibility() const;
 	
-	/** Handler for when the seleted parent class's filename is clicked */
+	/** Handler for when the selected parent class's filename is clicked */
 	void OnEditCodeClicked();
 
 	/** Handler for when a parent class item is double clicked */
@@ -181,9 +187,6 @@ private:
 	/** Returns the text for the calculated source file name */
 	FText OnGetClassSourceFileText() const;
 
-	/** Handler for when cancel is clicked */
-	void CancelClicked();
-
 	/** Returns true if Finish is allowed */
 	bool CanFinish() const;
 
diff --git a/Engine/Source/Runtime/AppFramework/Private/Widgets/Workflow/SWizard.cpp b/Engine/Source/Runtime/AppFramework/Private/Widgets/Workflow/SWizard.cpp
--- a/Engine/Source/Runtime/AppFramework/Private/Widgets/Workflow/SWizard.cpp
+++ b/Engine/Source/Runtime/AppFramework/Private/Widgets/Workflow/SWizard.cpp
@@ -174,7 +174,7 @@ void SWizard::Construct( const FArguments& InArgs )
 				+ SUniformGridPanel::Slot(2,0)
 				[
 					// 'Finish' button
-					SNew(SButton)
+					SAssignNew(FinishButton, SButton)
 					.ButtonStyle(InArgs._FinishButtonStyle)
 					.TextStyle(InArgs._ButtonTextStyle)
 					.HAlign(HAlign_Center)
@@ -348,6 +348,11 @@ void SWizard::AdvanceToPage(int32 PageIndex)
 	ShowPage(PageIndex);
 }
 
+TSharedPtr<SWidget> SWizard::GetWidgetToFocusOnActivate()
+{
+	return FinishButton;
+}
+
 FReply SWizard::HandleNextButtonClicked()
 {
 	int32 NextPage = GetNextPageIndex();
diff --git a/Engine/Source/Runtime/AppFramework/Public/Widgets/Workflow/SWizard.h b/Engine/Source/Runtime/AppFramework/Public/Widgets/Workflow/SWizard.h
--- a/Engine/Source/Runtime/AppFramework/Public/Widgets/Workflow/SWizard.h
+++ b/Engine/Source/Runtime/AppFramework/Public/Widgets/Workflow/SWizard.h
@@ -316,6 +316,9 @@ public:
 	/** Shows the page with the specified index, but preserves page history. */
 	APPFRAMEWORK_API void AdvanceToPage( int32 PageIndex );
 
+	/** Get the widget to focus on activation */
+	APPFRAMEWORK_API TSharedPtr<SWidget> GetWidgetToFocusOnActivate();
+
 public:
 
 	// SCompoundWidget interface
@@ -400,6 +403,8 @@ private:
 	// Whether to show the "Previous" and "Next" buttons
 	bool bShowHistoryButtons;
 
+	TSharedPtr<SWidget> FinishButton;
+
 private:
 
 	// Holds a delegate to be invoked when the 'Cancel' button has been clicked.
