[UBT][CppCompileWarnings] Bumping Game Project utils now that builds are stable.

diff --git a/Engine/Source/Editor/GameProjectGeneration/Private/GameProjectUtils.cpp b/Engine/Source/Editor/GameProjectGeneration/Private/GameProjectUtils.cpp
--- a/Engine/Source/Editor/GameProjectGeneration/Private/GameProjectUtils.cpp
+++ b/Engine/Source/Editor/GameProjectGeneration/Private/GameProjectUtils.cpp
@@ -2985,7 +2985,7 @@ bool GameProjectUtils::UpdateAdditionalPluginDirectory(const FString& InDir, con
 
 const TCHAR* GameProjectUtils::GetDefaultBuildSettingsVersion()
 {
-	return TEXT("BuildSettingsVersion.V5");
+	return TEXT("BuildSettingsVersion.V6");
 }
 
 bool GameProjectUtils::ReadTemplateFile(const FString& TemplateFileName, FString& OutFileContents, FText& OutFailReason)
