Fix: Show template projects from all plugins + update fallback image for template categories

Problem
New Project dialog only listed template projects from enabled plugins
Users couldn't see or start from templates included in installed but disabled plugins
Empty/invalid background image paths in template categories showed broken visuals

Solution
Updated search to include template projects from all discovered plugins
Added warning icon + tooltip for templates from disabled plugins
Selecting a disabled plugin template shows warning/error box that the plugin is not enabled
Warning is non-blocking (no modal), project creation proceeds after user confirmation

Notes
Keeps behavior lightweight: surfaces risk (no additional modal asking to confirm creation) but doesn't block flow
Validation: check New Project dialog with a disabled plugin installed, confirm template shows with warning, project creation works as expected

#13462

diff --git a/Engine/Plugins/Editor/PluginBrowser/Source/PluginBrowser/Private/PluginStyle.cpp b/Engine/Plugins/Editor/PluginBrowser/Source/PluginBrowser/Private/PluginStyle.cpp
--- a/Engine/Plugins/Editor/PluginBrowser/Source/PluginBrowser/Private/PluginStyle.cpp
+++ b/Engine/Plugins/Editor/PluginBrowser/Source/PluginBrowser/Private/PluginStyle.cpp
@@ -136,6 +136,9 @@ void FPluginStyle::Initialize()
 	const float ThumbnailImageSize = 69.0f;
 	StyleSet->Set( "PluginTile.ThumbnailImageSize", ThumbnailImageSize );
 
+	FString ThumbnailImageFallBack = IPluginManager::Get().FindPlugin(TEXT("PluginBrowser"))->GetBaseDir() / TEXT("Resources/DefaultIcon128.png");
+	StyleSet->Set( "PluginTile.ThumbnailImageFallBack", new FSlateImageBrush(ThumbnailImageFallBack, CoreStyleConstants::Icon128x128) );
+
 	StyleSet->Set( "PluginTile.BackgroundBrush", new FSlateNoResource );
 	StyleSet->Set( "PluginTile.BackgroundPadding", FMargin( PaddingAmount ) );
 
diff --git a/Engine/Plugins/Editor/PluginBrowser/Source/PluginBrowser/Private/SPluginTile.cpp b/Engine/Plugins/Editor/PluginBrowser/Source/PluginBrowser/Private/SPluginTile.cpp
--- a/Engine/Plugins/Editor/PluginBrowser/Source/PluginBrowser/Private/SPluginTile.cpp
+++ b/Engine/Plugins/Editor/PluginBrowser/Source/PluginBrowser/Private/SPluginTile.cpp
@@ -58,13 +58,10 @@ void SPluginTile::RecreateWidgets()
 
 	// Plugin thumbnail image
 	FString Icon128FilePath = Plugin->GetBaseDir() / TEXT("Resources/Icon128.png");
-	if(!FPlatformFileManager::Get().GetPlatformFile().FileExists(*Icon128FilePath))
-	{
-		Icon128FilePath = IPluginManager::Get().FindPlugin(TEXT("PluginBrowser"))->GetBaseDir() / TEXT("Resources/DefaultIcon128.png");
-	}
-
-	const FName BrushName( *Icon128FilePath );
+	
+	const FName BrushName(*Icon128FilePath);
 	const FIntPoint Size = FSlateApplication::Get().GetRenderer()->GenerateDynamicImageResource(BrushName);
+
 	if ((Size.X > 0) && (Size.Y > 0))
 	{
 		PluginIconDynamicImageBrush = MakeShareable(new FSlateDynamicImageBrush(BrushName, FVector2D(Size.X, Size.Y)));
@@ -373,7 +370,7 @@ void SPluginTile::RecreateWidgets()
 									.BorderImage(FPluginStyle::Get()->GetBrush("PluginTile.ThumbnailBorderImage"))
 									[
 										SNew(SImage)
-										.Image(PluginIconDynamicImageBrush.Get())
+										.Image(PluginIconDynamicImageBrush.IsValid() ? PluginIconDynamicImageBrush.Get() : FPluginStyle::Get()->GetBrush("PluginTile.ThumbnailImageFallBack"))
 									]
 								]
 								
diff --git a/Engine/Source/Editor/EditorStyle/Private/StarshipStyle.cpp b/Engine/Source/Editor/EditorStyle/Private/StarshipStyle.cpp
--- a/Engine/Source/Editor/EditorStyle/Private/StarshipStyle.cpp
+++ b/Engine/Source/Editor/EditorStyle/Private/StarshipStyle.cpp
@@ -7369,8 +7369,9 @@ void FStarshipEditorStyle::FStyle::SetupToolkitStyles()
 
 		Set("ProjectBrowser.TableRow", ProjectBrowserTableRowStyle);
 
-		Set("ProjectBrowser.MajorCategoryViewBorder", new FSlateRoundedBoxBrush(FStyleColors::Recessed, 4.0f));
-
+		Set("ProjectBrowser.MajorCategory.ViewBorder", new FSlateRoundedBoxBrush(FStyleColors::Recessed, 4.0f));
+		Set("ProjectBrowser.MajorCategory.Size", FVector2D(304, 104));
+		Set("ProjectBrowser.MajorCategory.FallBackThumbnail", new IMAGE_BRUSH("Starship/Projects/CustomTemplate_2x", GetVector("ProjectBrowser.MajorCategory.Size")));
 
 		Set("ProjectBrowser.TileViewTooltip.ToolTipBorder", new FSlateColorBrush(FLinearColor::Black));
 		Set("ProjectBrowser.TileViewTooltip.NonContentBorder", new BOX_BRUSH("/Docking/TabContentArea", FMargin(4 / 16.0f)));
diff --git a/Engine/Source/Editor/GameProjectGeneration/Private/GameProjectGenerationModule.cpp b/Engine/Source/Editor/GameProjectGeneration/Private/GameProjectGenerationModule.cpp
--- a/Engine/Source/Editor/GameProjectGeneration/Private/GameProjectGenerationModule.cpp
+++ b/Engine/Source/Editor/GameProjectGeneration/Private/GameProjectGenerationModule.cpp
@@ -182,7 +182,7 @@ void FGameProjectGenerationModule::LoadTemplateCategories()
 	TemplateRootFolders.Add(FPaths::EnterpriseDir() / TEXT("Templates"));
 
 	// allow plugins to define templates
-	TArray<TSharedRef<IPlugin>> Plugins = IPluginManager::Get().GetEnabledPlugins();
+	TArray<TSharedRef<IPlugin>> Plugins = IPluginManager::Get().GetDiscoveredPlugins();
 	for (const TSharedRef<IPlugin>& Plugin : Plugins)
 	{
 		FString PluginDirectory = Plugin->GetBaseDir();
@@ -221,8 +221,8 @@ void FGameProjectGenerationModule::LoadTemplateCategories()
 				TemplateCategory->Description = FLocalizedTemplateString::GetLocalizedText(Category.LocalizedDescriptions);
 				
 				const FName BrushName(*Category.Icon);
-				
-				FSlateDynamicImageBrush* Brush = new FSlateDynamicImageBrush(BrushName, FVector2D(300, 100));
+
+				TSharedPtr<FSlateDynamicImageBrush> Brush = MakeShared<FSlateDynamicImageBrush>(BrushName, FAppStyle::Get().GetVector("ProjectBrowser.MajorCategory.Size"));
 				Brush->OutlineSettings.CornerRadii = FVector4(4, 4, 4, 4);
 				Brush->OutlineSettings.RoundingType = ESlateBrushRoundingType::FixedRadius;
 				Brush->DrawAs = ESlateBrushDrawType::RoundedBox;
diff --git a/Engine/Source/Editor/GameProjectGeneration/Private/SProjectDialog.cpp b/Engine/Source/Editor/GameProjectGeneration/Private/SProjectDialog.cpp
--- a/Engine/Source/Editor/GameProjectGeneration/Private/SProjectDialog.cpp
+++ b/Engine/Source/Editor/GameProjectGeneration/Private/SProjectDialog.cpp
@@ -47,13 +47,12 @@
 #include "SPrimaryButton.h"
 #include "Styling/StyleColors.h"
 #include "Styling/AppStyle.h"
+#include "Styling/SlateStyleRegistry.h"
 
 #define LOCTEXT_NAMESPACE "GameProjectGeneration"
 
 namespace NewProjectDialogDefs
 {
-	constexpr float MajorItemWidth = 304;
-	constexpr float MajorItemHeight = 104;
 	constexpr float TemplateTileHeight = 153;
 	constexpr float TemplateTileWidth = 102;
 	constexpr float ThumbnailSize = 64.0f, ThumbnailPadding = 5.f;
@@ -61,7 +60,7 @@ namespace NewProjectDialogDefs
 	const FName BlankCategoryKey = "Default";
 }
 
-TUniquePtr<FSlateBrush> SProjectDialog::CustomTemplateBrush;
+TSharedPtr<FSlateBrush> SProjectDialog::CustomTemplateBrush;
 
 class SMajorCategoryTile : public SCompoundWidget
 {
@@ -73,6 +72,7 @@ public:
 	void Construct(const FArguments& InArgs, TSharedPtr<FTemplateCategory> Item)
 	{
 		IsSelected = InArgs._IsSelected;
+		TemplateCategoryWeak = Item;
 
 		ChildSlot
 		[
@@ -80,7 +80,7 @@ public:
 			+ SOverlay::Slot()
 			[
 				SNew(SImage)
-				.Image(Item->Icon)
+				.Image(this, &SMajorCategoryTile::GetIconBrush)
 			]
 			+ SOverlay::Slot()
 			.HAlign(HAlign_Left)
@@ -129,8 +129,23 @@ private:
 
 		return FStyleDefaults::GetNoBrush();
 	}
+
+	const FSlateBrush* GetIconBrush() const
+	{
+		if (TSharedPtr<FTemplateCategory> TemplateCategory = TemplateCategoryWeak.Pin())
+		{
+			if (TSharedPtr<FSlateBrush> Icon = TemplateCategory->Icon; Icon && FPaths::FileExists(Icon->GetResourceName().ToString()))
+			{
+				return Icon.Get();
+			}
+		}
+
+		return FAppStyle::Get().GetBrush("ProjectBrowser.MajorCategory.FallBackThumbnail");
+	}
+
 private:
 	TAttribute<bool> IsSelected;
+	TWeakPtr<FTemplateCategory> TemplateCategoryWeak;
 };
 
 
@@ -166,6 +181,8 @@ public:
 		STableRow::FArguments TableRowArguments;
 		TableRowArguments._SignalSelectionMode = ETableRowSignalSelectionMode::Instantaneous;
 
+		TSharedPtr<SOverlay> ProjectTileOverlay;
+
 		STableRow::Construct(
 			TableRowArguments
 			.Style(FAppStyle::Get(), "ProjectBrowser.TableRow")
@@ -176,7 +193,7 @@ public:
 				.Padding(FMargin(0.0f, 0.0f, 5.0f, 5.0f))
 				.BorderImage(FAppStyle::Get().GetBrush("ProjectBrowser.ProjectTile.DropShadow"))
 				[
-					SNew(SOverlay)
+					SAssignNew(ProjectTileOverlay, SOverlay)
 					+ SOverlay::Slot()
 					[
 						SNew(SVerticalBox)
@@ -294,6 +311,23 @@ public:
 			],
 			OwnerTable
 		);
+
+		if (const TSharedPtr<IPlugin> Plugin = InArgs._Item->PluginWeak.Pin(); Plugin && !Plugin->IsEnabled())
+		{
+			if (ProjectTileOverlay.IsValid())
+			{
+				ProjectTileOverlay->AddSlot()
+				.HAlign(HAlign_Right)
+				.VAlign(VAlign_Bottom)
+				.Padding(2)
+				[
+					SNew(SImage)
+						.Image(FAppStyle::GetBrush("Icons.Warning"))
+						.ToolTipText(LOCTEXT("WarnNotEnabled", "This template comes from a plugin that is not enabled by default."))
+						.ColorAndOpacity(EStyleColor::Warning)
+				];
+			}
+		}
 	}
 
 private:
@@ -664,97 +698,116 @@ TSharedRef<SWidget> SProjectDialog::MakeTemplateProjectView()
 			.Padding(8.0f, 0.0f)
 			.AutoWidth()
 			[
-				SNew(SVerticalBox)
-				.Clipping(EWidgetClipping::ClipToBounds)
-				// Preview image
-				+ SVerticalBox::Slot()
-				.AutoHeight()
-				.Padding(FMargin(0.0f, 0.0f, 0.0f, 15.f))
-				[
-					SNew(SImage)
-					.DesiredSizeOverride(FVector2D(358,160))
-					.Image(this, &SProjectDialog::GetSelectedTemplatePreviewImage)
-				]
-				// Template Name
-				+ SVerticalBox::Slot()
-				.Padding(FMargin(0.0f, 0.0f, 0.0f, 10.0f))
-				.AutoHeight()
+				SNew(SScrollBox)
+				.Orientation(Orient_Vertical)
+				.ConsumeMouseWheel(EConsumeMouseWheel::WhenScrollingPossible)
+				+ SScrollBox::Slot()
+				.FillSize(1.f)
 				[
-					SNew(STextBlock)
-					.AutoWrapText(true)
-					.TextStyle(FAppStyle::Get(), "DialogButtonText")
-					.Font(FAppStyle::Get().GetFontStyle("HeadingExtraSmall"))
-					.ColorAndOpacity(FAppStyle::Get().GetSlateColor("Colors.White"))
-					.Text(this, &SProjectDialog::GetSelectedTemplateProperty, &FTemplateItem::Name)
-				]
-				// Template Description
-				+ SVerticalBox::Slot()
-				.AutoHeight()
-				[
-					SNew(SBox)
-					.HeightOverride(120.0f)
-					.WidthOverride(358.0f)
+					SNew(SVerticalBox)
+					.Clipping(EWidgetClipping::ClipToBounds)
+					+ SVerticalBox::Slot()
+					.FillHeight(1.f)
+					.Padding(FMargin(0.0f, 0.0f, 0.0f, 15.f))
 					[
-						SNew(SScrollBox)
-						+ SScrollBox::Slot()
+						// Preview image
+						SNew(SVerticalBox)
+						+ SVerticalBox::Slot()
+						.AutoHeight()
+						.Padding(FMargin(0.0f, 0.0f, 0.0f, 15.f))
 						[
-							SNew(SVerticalBox)
-							+SVerticalBox::Slot()
-							.Padding(FMargin(0.0f, 0.0f, 0.0f, 10.0f))
-							.AutoHeight()
-							[
-								SNew(STextBlock)
-								.WrapTextAt(350.0f)
-								.Text(this, &SProjectDialog::GetSelectedTemplateProperty, &FTemplateItem::Description)
-							]
-							// Asset types
-							+ SVerticalBox::Slot()
-							.AutoHeight()
-							.Padding(FMargin(0.0f, 5.0f, 0.0f, 5.0f))
-							[
-								SNew(SVerticalBox)
-								.Visibility(this, &SProjectDialog::GetSelectedTemplateAssetVisibility)
-								+ SVerticalBox::Slot()
-								[
-									SNew(STextBlock)
-									.TextStyle(FAppStyle::Get(), "DialogButtonText")
-									.Text(LOCTEXT("ProjectTemplateAssetTypes", "Asset Type References"))
-								]
-								+ SVerticalBox::Slot()
-								.AutoHeight()
-								[
-									SNew(STextBlock)
-									.AutoWrapText(true)
-									.Text(this, &SProjectDialog::GetSelectedTemplateAssetTypes)
-								]
-							]
-							// Class types
-							+ SVerticalBox::Slot()
-							.AutoHeight()
-							.Padding(FMargin(0.0f, 5.0f, 0.0f, 5.0f))
+							SNew(SImage)
+							.DesiredSizeOverride(FVector2D(358, 160))
+							.Image(this, &SProjectDialog::GetSelectedTemplatePreviewImage)
+						]
+						// Template Name
+						+ SVerticalBox::Slot()
+						.Padding(FMargin(0.0f, 0.0f, 0.0f, 10.0f))
+						.AutoHeight()
+						[
+							SNew(STextBlock)
+							.AutoWrapText(true)
+							.TextStyle(FAppStyle::Get(), "DialogButtonText")
+							.Font(FAppStyle::Get().GetFontStyle("HeadingExtraSmall"))
+							.ColorAndOpacity(FAppStyle::Get().GetSlateColor("Colors.White"))
+							.Text(this, &SProjectDialog::GetSelectedTemplateProperty, &FTemplateItem::Name)
+						]
+						// Template Description
+						+ SVerticalBox::Slot()
+						.AutoHeight()
+						[
+							SNew(SBox)
+							.HeightOverride(120.0f)
+							.WidthOverride(358.0f)
 							[
-								SNew(SVerticalBox)
-								.Visibility(this, &SProjectDialog::GetSelectedTemplateClassVisibility)
-								+ SVerticalBox::Slot()
-								[
-									SNew(STextBlock)
-									.TextStyle(FAppStyle::Get(), "DialogButtonText")
-									.Text(LOCTEXT("ProjectTemplateClassTypes", "Class Type References"))
-								]
-								+ SVerticalBox::Slot()
-								.AutoHeight()
-								[
-									SNew(STextBlock)
-									.AutoWrapText(true)
-									.Text(this, &SProjectDialog::GetSelectedTemplateClassTypes)
+								SNew(SScrollBox)
+								+ SScrollBox::Slot()
+								[ 
+									SNew(SVerticalBox)
+									+ SVerticalBox::Slot()
+									.Padding(FMargin(0.0f, 0.0f, 0.0f, 10.0f))
+									.AutoHeight()
+									[
+										SNew(STextBlock)
+										.WrapTextAt(350.0f)
+										.Text(this, &SProjectDialog::GetSelectedTemplateProperty, &FTemplateItem::Description)
+									]
+									// Asset types
+									+ SVerticalBox::Slot()
+									.AutoHeight()
+									.Padding(FMargin(0.0f, 5.0f, 0.0f, 5.0f))
+									[
+										SNew(SVerticalBox)
+										.Visibility(this, &SProjectDialog::GetSelectedTemplateAssetVisibility)
+										+ SVerticalBox::Slot()
+										[
+											SNew(STextBlock)
+											.TextStyle(FAppStyle::Get(), "DialogButtonText")
+											.Text(LOCTEXT("ProjectTemplateAssetTypes", "Asset Type References"))
+										]
+										+ SVerticalBox::Slot()
+										.AutoHeight()
+										[
+											SNew(STextBlock)
+											.AutoWrapText(true)
+											.Text(this, &SProjectDialog::GetSelectedTemplateAssetTypes)
+										]
+									]
+									// Class types
+									+ SVerticalBox::Slot()
+									.AutoHeight()
+									.Padding(FMargin(0.0f, 5.0f, 0.0f, 5.0f))
+									[
+										SNew(SVerticalBox)
+										.Visibility(this, &SProjectDialog::GetSelectedTemplateClassVisibility)
+										+ SVerticalBox::Slot()
+										[
+											SNew(STextBlock)
+											.TextStyle(FAppStyle::Get(), "DialogButtonText")
+											.Text(LOCTEXT("ProjectTemplateClassTypes", "Class Type References"))
+										]
+										+ SVerticalBox::Slot()
+										.AutoHeight()
+										[
+											SNew(STextBlock)
+											.AutoWrapText(true)
+											.Text(this, &SProjectDialog::GetSelectedTemplateClassTypes)
+										]
+									]
 								]
 							]
 						]
+						// Project Options
+						+ SVerticalBox::Slot()
+						.Padding(FMargin(0.0f, 0.0f, 0.0f, 15.0f))
+						.AutoHeight()
+						.Expose(ProjectOptionsSlot)
+						// Plugin Information
+						+ SVerticalBox::Slot()
+						.AutoHeight()
+						.Expose(PluginInformationSlot)
 					]
 				]
-				// Project Options
-				+ SVerticalBox::Slot()
-				.Expose(ProjectOptionsSlot)
 			]
 		];
 }
@@ -778,6 +831,8 @@ TSharedRef<SWidget> SProjectDialog::MakeHybridView(EProjectDialogModeMode Mode)
 		.ItemWidth(NewProjectDialogDefs::TemplateTileWidth+9)
 		.OnSelectionChanged(this, &SProjectDialog::HandleTemplateListViewSelectionChanged);
 
+	const UE::Slate::FDeprecateVector2DResult MajorCategorySize = FAppStyle::Get().GetVector("ProjectBrowser.MajorCategory.Size");
+
 	TSharedRef<SWidget> HybridView = 
 		SNew(SHorizontalBox)
 		+ SHorizontalBox::Slot()
@@ -799,15 +854,15 @@ TSharedRef<SWidget> SProjectDialog::MakeHybridView(EProjectDialogModeMode Mode)
 			[
 				SNew(SBorder)
 				.Visibility(Mode == EProjectDialogModeMode::OpenProject ? EVisibility::Collapsed : EVisibility::Visible)
-				.BorderImage(FAppStyle::Get().GetBrush("ProjectBrowser.MajorCategoryViewBorder"))
+				.BorderImage(FAppStyle::Get().GetBrush("ProjectBrowser.MajorCategory.ViewBorder"))
 				[
 					SAssignNew(MajorCategoryList, STileView<TSharedPtr<FTemplateCategory>>)
 					.ListItemsSource(&TemplateCategories)
 					.SelectionMode(ESelectionMode::Single)
 					.ClearSelectionOnClick(false)
 					.OnGenerateTile(this, &SProjectDialog::ConstructMajorCategoryTableRow)
-					.ItemHeight(NewProjectDialogDefs::MajorItemHeight)
-					.ItemWidth(NewProjectDialogDefs::MajorItemWidth)
+					.ItemHeight(MajorCategorySize.Y)
+					.ItemWidth(MajorCategorySize.X)
 					.OnSelectionChanged(this, &SProjectDialog::OnMajorTemplateCategorySelectionChanged)
 				]
 			]
@@ -860,7 +915,7 @@ TSharedRef<SWidget> SProjectDialog::MakeProjectOptionsWidget()
 			]
 		]
 		+ SVerticalBox::Slot()
-		.Padding(0.0f, 20.0f, 0.0f, 0.0f)
+		.Padding(0.0f, 10.f, 0.0f, 0.0f)
 		[
 			SNew(SScrollBox)
 			+ SScrollBox::Slot()
@@ -949,7 +1004,8 @@ TSharedRef<SWidget> SProjectDialog::MakeProjectOptionsWidget()
 				]
 			]
 		];
-	}
+	}	
+
 
 	if (!HiddenSettings.Contains(ETemplateSetting::GraphicsPreset))
 	{
@@ -983,7 +1039,8 @@ TSharedRef<SWidget> SProjectDialog::MakeProjectOptionsWidget()
 	}
 	
 	ResetVariants(); // Needs to be called after we have set bShouldGenerateCode
-	if (!HiddenSettings.Contains(ETemplateSetting::Variants))
+	const UTemplateProjectDefs* ProjectDefs = GetSelectedTemplateDefs();
+	if (!HiddenSettings.Contains(ETemplateSetting::Variants) && ProjectDefs && !ProjectDefs->Variants.IsEmpty())
 	{
 		ProjectOptionsBox->AddSlot()
 			.Padding(0.0f, 8.0f, 0.0f, 0.0f)
@@ -1042,9 +1099,121 @@ TSharedRef<SWidget> SProjectDialog::MakeProjectOptionsWidget()
 	}
 #endif
 
+	if (ProjectOptionsBox->NumSlots() == 0)
+	{
+		return SNullWidget::NullWidget;
+	}
+
 	return ProjectOptionsWidget;
 }
 
+TSharedRef<SWidget> SProjectDialog::MakePluginInformationWidget()
+{
+	const TSharedPtr<IPlugin> Plugin = GetSelectedTemplateProperty(&FTemplateItem::PluginWeak).Pin();
+	if (!Plugin)
+	{
+		return SNullWidget::NullWidget;
+	}
+
+	const FPluginDescriptor& PluginDescriptor = Plugin->GetDescriptor();
+
+	const ISlateStyle* PluginStyle = FSlateStyleRegistry::FindSlateStyle("PluginStyle");
+
+	const float PaddingAmount = PluginStyle->GetFloat("PluginTile.Padding");
+	const float ThumbnailImageSize = PluginStyle->GetFloat("PluginTile.ThumbnailImageSize");
+
+	// Plugin thumbnail image
+	FString Icon128FilePath = Plugin->GetBaseDir() / TEXT("Resources/Icon128.png");
+
+	const FName BrushName(*Icon128FilePath);
+	const FIntPoint Size = FSlateApplication::Get().GetRenderer()->GenerateDynamicImageResource(BrushName);
+
+	if ((Size.X > 0) && (Size.Y > 0))
+	{
+		PluginIconDynamicImageBrush = MakeShareable(new FSlateDynamicImageBrush(BrushName, FVector2D(Size.X, Size.Y)));
+	}
+
+	TSharedRef<SWidget> PluginInformationWidget =
+		SNew(SVerticalBox)
+		.Visibility(this, &SProjectDialog::GetProjectSettingsVisibility)
+		// Plugin information header
+		+ SVerticalBox::Slot()
+		.AutoHeight()
+		[
+			SNew(SBorder)
+			.Padding(FMargin(10.0f, 7.0f))
+			.BorderImage(FAppStyle::Get().GetBrush("Brushes.Header"))
+			[
+				SNew(STextBlock)
+				.TextStyle(FAppStyle::Get(), "DialogButtonText")
+				.Text(LOCTEXT("PluginInformation", "Plugin Information"))
+			]
+		]
+		// Plugin tile 
+		+ SVerticalBox::Slot()
+		.Padding(0.0f, 10.f, 0.0f, 0.0f)
+		[
+			SNew(SVerticalBox)
+			+SVerticalBox::Slot()
+			.Padding(0)
+			.AutoHeight()
+			[
+				SNew(SBorder)
+				.BorderImage(FAppStyle::Get().GetBrush("Brushes.Panel"))
+				.Padding(0)
+				[
+					SNew(SBorder)
+					.BorderImage(PluginStyle->GetBrush("PluginTile.BorderImage"))
+					.Padding(PaddingAmount + 4.f)
+					[
+						SNew(SHorizontalBox)
+						// Thumbnail image
+						+ SHorizontalBox::Slot()
+						.Padding(PaddingAmount, PaddingAmount, PaddingAmount + 10.f, PaddingAmount)
+						.VAlign(VAlign_Center)
+						.AutoWidth()
+						[
+							SNew(SBox)
+							.WidthOverride(ThumbnailImageSize)
+							.HeightOverride(ThumbnailImageSize)
+							[
+								SNew(SBorder)
+								.BorderImage(PluginStyle->GetBrush("PluginTile.ThumbnailBorderImage"))
+								[
+									SNew(SImage)
+									.Image(PluginIconDynamicImageBrush.IsValid() ? PluginIconDynamicImageBrush.Get() : PluginStyle->GetBrush("PluginTile.ThumbnailImageFallBack"))
+								]
+							]
+						]
+						+ SHorizontalBox::Slot()
+						[
+							SNew(SVerticalBox)
+							// Friendly Name
+							+ SVerticalBox::Slot()
+							.AutoHeight()
+							.Padding(PaddingAmount, PaddingAmount + 3.f, PaddingAmount + 8.f, 0.f)
+							[
+								SNew(STextBlock)
+								.Text(FText::FromString(Plugin->GetFriendlyName()))
+								.TextStyle(PluginStyle, "PluginTile.NameText")
+							]
+							// Created By
+							+ SVerticalBox::Slot()
+							.AutoHeight()
+							.Padding(PaddingAmount, PaddingAmount, 0, 0)
+							[
+								SNew(STextBlock)
+								.Text(FText::FromString(PluginDescriptor.CreatedBy))
+							]
+						]
+					]
+				]
+			]
+		];
+
+	return PluginInformationWidget;
+}
+
 TSharedRef<SWidget> SProjectDialog::MakeRecentProjectsTile()
 {
 	RecentProjectsCategory = MakeShared<FTemplateCategory>();
@@ -1056,12 +1225,12 @@ TSharedRef<SWidget> SProjectDialog::MakeRecentProjectsTile()
 
 	static const FName BrushName =  *(FAppStyle::Get().GetContentRootDir() / TEXT("/Starship/Projects/") / TEXT("RecentProjects_2x.png"));
 
-	RecentProjectsBrush = MakeUnique<FSlateDynamicImageBrush>(BrushName, FVector2D(300, 100));
+	RecentProjectsBrush = MakeShared<FSlateDynamicImageBrush>(BrushName, FVector2D(300, 100));
 	RecentProjectsBrush->OutlineSettings.CornerRadii = FVector4(4, 4, 4, 4);
 	RecentProjectsBrush->OutlineSettings.RoundingType = ESlateBrushRoundingType::FixedRadius;
 	RecentProjectsBrush->DrawAs = ESlateBrushDrawType::RoundedBox;
 
-	RecentProjectsCategory->Icon = RecentProjectsBrush.Get();
+	RecentProjectsCategory->Icon = RecentProjectsBrush;
 
 	return
 		SNew(SButton)
@@ -1097,9 +1266,21 @@ TSharedRef<SWidget> SProjectDialog::MakeNewProjectPathArea()
 		.AutoHeight()
 		[
 			SNew(SHorizontalBox)
+			+ SHorizontalBox::Slot()
+			.VAlign(VAlign_Center)
+			.AutoWidth()
+			[
+				SNew(SWarningOrErrorBox)
+				.Padding(FMargin(8.0f, 4.0f, 4.0f, 4.0f))
+				.IconSize(FVector2D(16, 16))
+				.MessageStyle(EMessageStyle::Warning)
+				.Message(this, &SProjectDialog::GetGlobalWarningLabelText)
+				.Visibility(this, &SProjectDialog::GetGlobalWarningVisibility)
+			]
 			+SHorizontalBox::Slot()
 			.VAlign(VAlign_Center)
 			.AutoWidth()
+			.Padding(FMargin(8.0f, 0.0f, 0.0f, 0.0f))
 			[
 				SNew(SWarningOrErrorBox)
 				.Padding(FMargin(8.0f, 4.0f, 4.0f, 4.0f))
@@ -1255,7 +1436,15 @@ void SProjectDialog::HandleTemplateListViewSelectionChanged(TSharedPtr<FTemplate
 		MakeProjectOptionsWidget()
 	];
 
+	(*PluginInformationSlot)
+	[
+		MakePluginInformationWidget()
+	];
+	
+
 	UpdateProjectFileValidity();
+
+	UpdateTemplateTrustWarning();
 }
 
 TSharedPtr<FTemplateItem> SProjectDialog::GetSelectedTemplateItem() const
@@ -1304,17 +1493,22 @@ TMap<FName, TArray<TSharedPtr<FTemplateItem>> > SProjectDialog::FindTemplateProj
 	// Clear the list out first - or we could end up with duplicates
 	TMap<FName, TArray<TSharedPtr<FTemplateItem>>> Templates;
 
+	struct FTemplateFolderMetadata
+	{
+		FString Folder;
+		TWeakPtr<IPlugin> PluginWeak;
+	};
 	// Now discover and all data driven templates
-	TArray<FString> TemplateRootFolders;
+	TArray<FTemplateFolderMetadata> TemplateRootFolders;
 
 	// @todo rocket make template folder locations extensible.
-	TemplateRootFolders.Add(FPaths::RootDir() + TEXT("Templates"));
+	TemplateRootFolders.Add({ FPaths::RootDir() + TEXT("Templates"), nullptr });
 
 	// Add the Enterprise templates
-	TemplateRootFolders.Add(FPaths::EnterpriseDir() + TEXT("Templates"));
+	TemplateRootFolders.Add({ FPaths::EnterpriseDir() + TEXT("Templates"), nullptr });
 
 	// Allow plugins to define templates
-	TArray<TSharedRef<IPlugin>> Plugins = IPluginManager::Get().GetEnabledPlugins();
+	TArray<TSharedRef<IPlugin>> Plugins = IPluginManager::Get().GetDiscoveredPlugins();
 	for (const TSharedRef<IPlugin>& Plugin : Plugins)
 	{
 		FString PluginDirectory = Plugin->GetBaseDir();
@@ -1324,31 +1518,32 @@ TMap<FName, TArray<TSharedPtr<FTemplateItem>> > SProjectDialog::FindTemplateProj
 
 			if (IFileManager::Get().DirectoryExists(*PluginTemplatesDirectory))
 			{
-				TemplateRootFolders.Add(PluginTemplatesDirectory);
+				TemplateRootFolders.Add({ PluginTemplatesDirectory, Plugin.ToWeakPtr() });
 			}
 		}
 	}
 
+
 	// Form a list of all folders that could contain template projects
-	TArray<FString> AllTemplateFolders;
-	for (const FString& Root : TemplateRootFolders)
+	TArray<FTemplateFolderMetadata> AllTemplateFolders;
+	for (const auto& [RootFolder, Plugin] : TemplateRootFolders)
 	{
-		const FString SearchString = Root / TEXT("*");
+		const FString SearchString = RootFolder / TEXT("*");
 		TArray<FString> TemplateFolders;
 		IFileManager::Get().FindFiles(TemplateFolders, *SearchString, /*Files=*/false, /*Directories=*/true);
 
 		for (const FString& Folder : TemplateFolders)
 		{
-			AllTemplateFolders.Add(Root / Folder);
+			AllTemplateFolders.Add({ RootFolder / Folder, Plugin });
 		}
 	}
 
 	TArray<TSharedPtr<FTemplateItem>> FoundTemplates;
 
 	// Add a template item for every discovered project
-	for (const FString& Root : AllTemplateFolders)
+	for (const auto& [Folder, Plugin] : AllTemplateFolders)
 	{
-		const FString SearchString = Root / TEXT("*.") + FProjectDescriptor::GetExtension();
+		const FString SearchString = Folder / TEXT("*.") + FProjectDescriptor::GetExtension();
 		TArray<FString> FoundProjectFiles;
 		IFileManager::Get().FindFiles(FoundProjectFiles, *SearchString, /*Files=*/true, /*Directories=*/false);
 
@@ -1358,7 +1553,7 @@ TMap<FName, TArray<TSharedPtr<FTemplateItem>> > SProjectDialog::FindTemplateProj
 		}
 
 		// Make sure a TemplateDefs.ini file exists
-		UTemplateProjectDefs* TemplateDefs = GameProjectUtils::LoadTemplateDefs(Root);
+		UTemplateProjectDefs* TemplateDefs = GameProjectUtils::LoadTemplateDefs(Folder);
 		if (TemplateDefs == nullptr)
 		{
 			continue;
@@ -1373,7 +1568,7 @@ TMap<FName, TArray<TSharedPtr<FTemplateItem>> > SProjectDialog::FindTemplateProj
 			continue;
 		}
 
-		const FString ProjectFile = Root / FoundProjectFiles[0];
+		const FString ProjectFile = Folder / FoundProjectFiles[0];
 
 		// If no template category was specified, use the default category
 		TArray<FName> TemplateCategoryNames = TemplateDefs->Categories;
@@ -1383,7 +1578,7 @@ TMap<FName, TArray<TSharedPtr<FTemplateItem>> > SProjectDialog::FindTemplateProj
 		}
 
 		// Find a duplicate project, eg. "Foo" and "FooBP"
-		FString TemplateKey = Root;
+		FString TemplateKey = Folder;
 		TemplateKey.RemoveFromEnd("BP");
 
 		TSharedPtr<FTemplateItem>* ExistingTemplate = FoundTemplates.FindByPredicate([&TemplateKey](TSharedPtr<FTemplateItem> Item)
@@ -1403,7 +1598,7 @@ TMap<FName, TArray<TSharedPtr<FTemplateItem>> > SProjectDialog::FindTemplateProj
 			Template = MakeShareable(new FTemplateItem());
 		}
 
-		if (TemplateDefs->GeneratesCode(Root))
+		if (TemplateDefs->GeneratesCode(Folder))
 		{
 			Template->CodeProjectFile = ProjectFile;
 			Template->CodeTemplateDefs = TemplateDefs;
@@ -1432,6 +1627,7 @@ TMap<FName, TArray<TSharedPtr<FTemplateItem>> > SProjectDialog::FindTemplateProj
 		Template->bIsEnterprise = TemplateDefs->bIsEnterprise;
 		Template->bIsBlankTemplate = TemplateDefs->bIsBlank;
 		Template->bThumbnailAsIcon = TemplateDefs->bThumbnailAsIcon;
+		Template->PluginWeak = Plugin;
 
 		Template->Name = TemplateDefs->GetDisplayNameText();
 		if (Template->Name.IsEmpty())
@@ -1439,7 +1635,7 @@ TMap<FName, TArray<TSharedPtr<FTemplateItem>> > SProjectDialog::FindTemplateProj
 			Template->Name = FText::FromString(TemplateKey);
 		}
 
-		const FString ThumbnailPNGFile = (Root + TEXT("/Media/") + FoundProjectFiles[0]).Replace(TEXT(".uproject"), TEXT(".png"));
+		const FString ThumbnailPNGFile = (Folder + TEXT("/Media/") + FoundProjectFiles[0]).Replace(TEXT(".uproject"), TEXT(".png"));
 		if (FPlatformFileManager::Get().GetPlatformFile().FileExists(*ThumbnailPNGFile))
 		{
 			const FName BrushName = FName(*ThumbnailPNGFile);
@@ -1447,7 +1643,7 @@ TMap<FName, TArray<TSharedPtr<FTemplateItem>> > SProjectDialog::FindTemplateProj
 		}
 
 		TSharedPtr<FSlateDynamicImageBrush> PreviewBrush;
-		const FString PreviewPNGFile = (Root + TEXT("/Media/") + FoundProjectFiles[0]).Replace(TEXT(".uproject"), TEXT("_Preview.png"));
+		const FString PreviewPNGFile = (Folder + TEXT("/Media/") + FoundProjectFiles[0]).Replace(TEXT(".uproject"), TEXT("_Preview.png"));
 		if (FPlatformFileManager::Get().GetPlatformFile().FileExists(*PreviewPNGFile))
 		{
 			const FName BrushName = FName(*PreviewPNGFile);
@@ -1773,6 +1969,13 @@ void SProjectDialog::DisplayError(const FText& ErrorText)
 	}
 }
 
+void SProjectDialog::DisplayWarning(const FText& WarningText)
+{
+	FString WarningString = WarningText.ToString();
+	UE_LOG(LogGameProjectGeneration, Log, TEXT("%s"), *WarningString);
+	PersistentGlobalWarningLabelText = WarningText;
+}
+
 bool SProjectDialog::OpenCodeIDE(const FString& ProjectFile)
 {
 	FText FailReason;
@@ -1822,7 +2025,7 @@ TArray<TSharedPtr<FTemplateCategory>> SProjectDialog::GetAllTemplateCategories()
 
 		if (!CustomTemplateBrush)
 		{
-			CustomTemplateBrush = MakeUnique<FSlateDynamicImageBrush>(BrushName, FVector2D(300, 100));
+			CustomTemplateBrush = MakeShared<FSlateDynamicImageBrush>(BrushName, FVector2D(300, 100));
 			CustomTemplateBrush->OutlineSettings.CornerRadii = FVector4(4, 4, 4, 4);
 			CustomTemplateBrush->OutlineSettings.RoundingType = ESlateBrushRoundingType::FixedRadius;
 			CustomTemplateBrush->DrawAs = ESlateBrushDrawType::RoundedBox;
@@ -1832,7 +2035,7 @@ TArray<TSharedPtr<FTemplateCategory>> SProjectDialog::GetAllTemplateCategories()
 		DefaultCategory->Key = NewProjectDialogDefs::BlankCategoryKey;
 		DefaultCategory->DisplayName = LOCTEXT("ProjectDialog_DefaultCategoryName", "Blank Project");
 		DefaultCategory->Description = LOCTEXT("ProjectDialog_DefaultCategoryDescription", "Create a new blank Unreal project.");
-		DefaultCategory->Icon = CustomTemplateBrush.Get();
+		DefaultCategory->Icon = CustomTemplateBrush;
 
 		AllTemplateCategories.Add(DefaultCategory);
 	}
@@ -1857,6 +2060,11 @@ FText SProjectDialog::GetGlobalErrorLabelText() const
 	return FText::GetEmpty();
 }
 
+FText SProjectDialog::GetGlobalWarningLabelText() const
+{
+	return PersistentGlobalWarningLabelText;
+}
+
 FText SProjectDialog::GetNameAndLocationValidityErrorText() const
 {
 	if (GetGlobalErrorLabelText().IsEmpty())
@@ -2020,8 +2228,8 @@ void SProjectDialog::UpdateProjectFileValidity()
 	{
 		bLastGlobalValidityCheckSuccessful = true;
 
-		TSharedPtr<FTemplateItem> SelectedTemplate = GetSelectedTemplateItem();
-		if (!SelectedTemplate.IsValid())
+		
+		if (TSharedPtr<FTemplateItem> SelectedTemplate = GetSelectedTemplateItem(); !SelectedTemplate.IsValid())
 		{
 			bLastGlobalValidityCheckSuccessful = false;
 			LastGlobalValidityErrorText = LOCTEXT("NoTemplateSelected", "No Template Selected");
@@ -2089,6 +2297,20 @@ void SProjectDialog::UpdateProjectFileValidity()
 	}
 }
 
+void SProjectDialog::UpdateTemplateTrustWarning()
+{
+	if (TSharedPtr<FTemplateItem> SelectedTemplate = GetSelectedTemplateItem(); SelectedTemplate.IsValid() && SelectedTemplate->PluginWeak != nullptr)
+	{
+		if (const TSharedPtr<IPlugin> Plugin = SelectedTemplate->PluginWeak.Pin(); !Plugin->IsEnabled())
+		{
+			DisplayWarning(LOCTEXT("WarnNotEnabledTrust", "This template comes from a plugin that is not enabled by default. Make sure you trust the authors."));
+			return;
+		}
+	}
+
+	PersistentGlobalWarningLabelText = FText::GetEmpty();
+}
+
 void SProjectDialog::OnMajorTemplateCategorySelectionChanged(TSharedPtr<FTemplateCategory> Item, ESelectInfo::Type SelectType)
 {
 	if(Item.IsValid())
diff --git a/Engine/Source/Editor/GameProjectGeneration/Private/SProjectDialog.h b/Engine/Source/Editor/GameProjectGeneration/Private/SProjectDialog.h
--- a/Engine/Source/Editor/GameProjectGeneration/Private/SProjectDialog.h
+++ b/Engine/Source/Editor/GameProjectGeneration/Private/SProjectDialog.h
@@ -108,6 +108,9 @@ private:
 	/** Checks the current project path an name for validity and updates cached values accordingly */
 	void UpdateProjectFileValidity();
 
+	/** Checks if the currently selected template is from a plugin that is not enabled and updates the warning message accordingly */
+	void UpdateTemplateTrustWarning();
+
 	/** Returns true if we have a code template selected */
 	bool IsCompilerRequired() const;
 
@@ -123,6 +126,8 @@ private:
 
 	EVisibility GetGlobalErrorVisibility() const { return GetGlobalErrorLabelText().IsEmpty() ? EVisibility::Hidden : EVisibility::HitTestInvisible; }
 
+	EVisibility GetGlobalWarningVisibility() const { return PersistentGlobalWarningLabelText.IsEmpty() ? EVisibility::Hidden : EVisibility::HitTestInvisible; }
+
 	EVisibility GetCreateButtonVisibility() const;
 
 	EVisibility GetSuggestedIDEButtonVisibility() const;
@@ -143,6 +148,8 @@ private:
 
 	FText GetGlobalErrorLabelText() const;
 
+	FText GetGlobalWarningLabelText() const;
+
 	FText GetNameAndLocationValidityErrorText() const;
 
 	/** Gets the assembled project filename with path */
@@ -160,6 +167,8 @@ private:
 
 	TSharedRef<SWidget> MakeProjectOptionsWidget();
 
+	TSharedRef<SWidget> MakePluginInformationWidget();
+
 	TSharedRef<SWidget> MakeRecentProjectsTile();
 
 	TSharedRef<SWidget> MakeNewProjectPathArea();
@@ -204,6 +213,7 @@ private:
 	bool OpenCodeIDE(const FString& ProjectFile);
 	void CloseWindowIfAppropriate(bool ForceClose = false);
 	void DisplayError(const FText& ErrorText);
+	void DisplayWarning(const FText& WarningText);
 
 
 	static TArray<TSharedPtr<FTemplateCategory>> GetAllTemplateCategories();
@@ -227,13 +237,16 @@ private:
 	TSharedPtr<SWidgetSwitcher> TemplateAndRecentProjectsSwitcher;
 	TSharedPtr<SWidgetSwitcher> PathAreaSwitcher;
 
-	TSharedPtr<SProjectBrowser> ProjectBrowser; 
+	TSharedPtr<SProjectBrowser> ProjectBrowser;
+
+	TSharedPtr<FSlateDynamicImageBrush> PluginIconDynamicImageBrush;
 
 	FString LastBrowsePath;
 	FString CurrentProjectFileName;
 	FString CurrentProjectFilePath;
 
 	FText PersistentGlobalErrorLabelText;
+	FText PersistentGlobalWarningLabelText;
 
 	/** The global error text from the last validity check */
 	FText LastGlobalValidityErrorText;
@@ -250,12 +263,14 @@ private:
 	FName SelectedVariantName;
 	TArray<FName> TemplateVariantNames;
 
-	TUniquePtr<FSlateBrush> RecentProjectsBrush;
+	TSharedPtr<FSlateBrush> RecentProjectsBrush;
 
-	static TUniquePtr<FSlateBrush> CustomTemplateBrush;
+	static TSharedPtr<FSlateBrush> CustomTemplateBrush;
 
 	SVerticalBox::FSlot* ProjectOptionsSlot;
 
+	SVerticalBox::FSlot* PluginInformationSlot;
+
 	/** Whether or not to enable XR in the created project. */
 	bool bEnableXR;
 
diff --git a/Engine/Source/Editor/GameProjectGeneration/Private/TemplateCategory.h b/Engine/Source/Editor/GameProjectGeneration/Private/TemplateCategory.h
--- a/Engine/Source/Editor/GameProjectGeneration/Private/TemplateCategory.h
+++ b/Engine/Source/Editor/GameProjectGeneration/Private/TemplateCategory.h
@@ -15,7 +15,7 @@ struct FTemplateCategory
 	FText Description;
 
 	/** A thumbnail to help identify this category (on the tab) */
-	const FSlateBrush* Icon;
+	TSharedPtr<FSlateBrush> Icon;
 
 	/** A unique key for this category */
 	FName Key;
diff --git a/Engine/Source/Editor/GameProjectGeneration/Private/TemplateItem.h b/Engine/Source/Editor/GameProjectGeneration/Private/TemplateItem.h
--- a/Engine/Source/Editor/GameProjectGeneration/Private/TemplateItem.h
+++ b/Engine/Source/Editor/GameProjectGeneration/Private/TemplateItem.h
@@ -5,6 +5,7 @@
 #include "CoreMinimal.h"
 #include "Styling/SlateBrush.h"
 #include "TemplateProjectDefs.h"
+#include "Interfaces/IPluginManager.h"
 
 /** Struct describing a single template project */
 struct FTemplateItem
@@ -18,6 +19,8 @@ struct FTemplateItem
 
 	TSharedPtr<FSlateBrush> Thumbnail;
 	TSharedPtr<FSlateBrush> PreviewImage;
+	TWeakPtr<IPlugin> PluginWeak;
+
 
 	FString		ClassTypes;
 	FString		AssetTypes;
@@ -32,4 +35,5 @@ struct FTemplateItem
 	bool		bIsEnterprise = false;
 	bool		bIsBlankTemplate = false;
 	bool		bThumbnailAsIcon = false;
+	
 };
