[Substrate] Enabled Substrate by default in new projects. Also removed beta tag on cvar and project setting.

diff --git a/Engine/Source/Editor/GameProjectGeneration/Private/GameProjectUtils.cpp b/Engine/Source/Editor/GameProjectGeneration/Private/GameProjectUtils.cpp
--- a/Engine/Source/Editor/GameProjectGeneration/Private/GameProjectUtils.cpp
+++ b/Engine/Source/Editor/GameProjectGeneration/Private/GameProjectUtils.cpp
@@ -195,6 +195,16 @@ namespace
 			false /* ShouldReplaceExistingValue */);
 	}
 
+	/** Get the configuration values for enabling Substrate by default. */
+	void AddSubstrateConfigValues(const FProjectInformation& InProjectInfo, TArray<FTemplateConfigValue>& ConfigValues)
+	{
+		ConfigValues.Emplace(TEXT("DefaultEngine.ini"),
+			TEXT("/Script/Engine.RendererSettings"),
+			TEXT("r.Substrate"),
+			TEXT("True"),
+			false /* ShouldReplaceExistingValue */);
+	}
+
 	/** Get the configuration values for enabling WorldPartition by default. */
 	void AddWorldPartitionConfigValues(const FProjectInformation& InProjectInfo, TArray<FTemplateConfigValue>& ConfigValues)
 	{
@@ -1853,6 +1863,7 @@ TOptional<FGuid> GameProjectUtils::CreateProjectFromTemplate(const FProjectInfor
 
 	AddLumenConfigValues(InProjectInfo, ConfigValuesToSet, /* do not disable static lighting, as some templates need it */ false);
 	AddRaytracingConfigValues(InProjectInfo, ConfigValuesToSet);
+	AddSubstrateConfigValues(InProjectInfo, ConfigValuesToSet);
 	AddNewProjectDefaultShadowConfigValues(InProjectInfo, ConfigValuesToSet);
 	AddPostProcessingConfigValues(InProjectInfo, ConfigValuesToSet);
 	AddWorldPartitionConfigValues(InProjectInfo, ConfigValuesToSet);
@@ -2171,6 +2182,7 @@ bool GameProjectUtils::GenerateConfigFiles(const FProjectInformation& InProjectI
 		AddNewProjectDefaultShadowConfigValues(InProjectInfo, ConfigValuesToSet);
 		AddPostProcessingConfigValues(InProjectInfo, ConfigValuesToSet);
 		AddRaytracingConfigValues(InProjectInfo, ConfigValuesToSet);
+		AddSubstrateConfigValues(InProjectInfo, ConfigValuesToSet);
 		AddWorldPartitionConfigValues(InProjectInfo, ConfigValuesToSet);
 		AddUserInterfaceEngineConfigValues(InProjectInfo, ConfigValuesToSet);
 
diff --git a/Engine/Source/Runtime/Engine/Classes/Engine/RendererSettings.h b/Engine/Source/Runtime/Engine/Classes/Engine/RendererSettings.h
--- a/Engine/Source/Runtime/Engine/Classes/Engine/RendererSettings.h
+++ b/Engine/Source/Runtime/Engine/Classes/Engine/RendererSettings.h
@@ -1202,11 +1202,11 @@ class URendererSettings : public UDeveloperSettings
 		uint32 bSupportCloudShadowOnSingleLayerWater : 1;
 
 	/**
-	"Enable Substrate materials (Beta)."
+	"Enable Substrate materials."
 	*/
 	UPROPERTY(config, EditAnywhere, Category = Substrate, meta = (
-		ConsoleVariable = "r.Substrate", DisplayName = "Substrate materials (Beta)",
-		ToolTip = "Enable Substrate materials (Beta).",
+		ConsoleVariable = "r.Substrate", DisplayName = "Substrate materials",
+		ToolTip = "Enable Substrate materials.",
 		ConfigRestartRequired = true))
 		uint32 bEnableSubstrate : 1;
 
diff --git a/Engine/Source/Runtime/RenderCore/Private/RenderUtils.cpp b/Engine/Source/Runtime/RenderCore/Private/RenderUtils.cpp
--- a/Engine/Source/Runtime/RenderCore/Private/RenderUtils.cpp
+++ b/Engine/Source/Runtime/RenderCore/Private/RenderUtils.cpp
@@ -1931,7 +1931,7 @@ uint32 GetHairStrandsLODMode()
 static TAutoConsoleVariable<int32> CVarSubstrate(
 	TEXT("r.Substrate"),
 	0,
-	TEXT("Enable Substrate materials (Beta)."),
+	TEXT("Enable Substrate materials."),
 	ECVF_ReadOnly | ECVF_RenderThreadSafe);
 
 static TAutoConsoleVariable<int32> CVarSubstrateProjectGBufferFormat(
