New Project Dialog
- Fixed size of preview images to be 400x200
- Made sure all thumbnails had consistant rounded corners

diff --git a/Engine/Source/Editor/EditorStyle/Private/StarshipStyle.cpp b/Engine/Source/Editor/EditorStyle/Private/StarshipStyle.cpp
--- a/Engine/Source/Editor/EditorStyle/Private/StarshipStyle.cpp
+++ b/Engine/Source/Editor/EditorStyle/Private/StarshipStyle.cpp
@@ -510,6 +510,12 @@ void FStarshipEditorStyle::FStyle::SetupGeneralStyles()
 
 	Set("UnrealDefaultThumbnail", new IMAGE_BRUSH("Starship/Common/Unreal_DefaultThumbnail", FVector2f(256.f, 256.f)));
 
+	FSlateImageBrush* RoundedDefaultThumbnail = new IMAGE_BRUSH("Starship/Common/Unreal_DefaultThumbnail", FVector2f(256.f, 256.f));
+	RoundedDefaultThumbnail->OutlineSettings.CornerRadii = FVector4(4, 4, 0, 0);
+	RoundedDefaultThumbnail->OutlineSettings.RoundingType = ESlateBrushRoundingType::FixedRadius;
+	RoundedDefaultThumbnail->DrawAs = ESlateBrushDrawType::RoundedBox;
+	Set("UnrealDefaultThumbnailRounded", RoundedDefaultThumbnail);
+
 	Set( "WarningStripe", new IMAGE_BRUSH( "Common/WarningStripe", FVector2f(20.f,6.f), FLinearColor::White, ESlateBrushTileType::Horizontal ) );
 	
 	Set("RoundedWarning", new FSlateRoundedBoxBrush(FStyleColors::Transparent, 4.0f, FStyleColors::Warning, 1.0f));
diff --git a/Engine/Source/Editor/GameProjectGeneration/Private/SProjectBrowser.cpp b/Engine/Source/Editor/GameProjectGeneration/Private/SProjectBrowser.cpp
--- a/Engine/Source/Editor/GameProjectGeneration/Private/SProjectBrowser.cpp
+++ b/Engine/Source/Editor/GameProjectGeneration/Private/SProjectBrowser.cpp
@@ -222,7 +222,7 @@ public:
 							.Padding(FMargin(0))
 							[
 								SNew(SImage)
-								.Image(ProjectItem->ProjectThumbnail ? ProjectItem->ProjectThumbnail.Get() : FAppStyle::Get().GetBrush("UnrealDefaultThumbnail"))
+								.Image(ProjectItem->ProjectThumbnail ? ProjectItem->ProjectThumbnail.Get() : FAppStyle::Get().GetBrush("UnrealDefaultThumbnailRounded"))
 								.ColorAndOpacity(FAppStyle::Get().GetSlateColor("Colors.Foreground"))
 							]
 						]
diff --git a/Engine/Source/Editor/GameProjectGeneration/Private/SProjectDialog.cpp b/Engine/Source/Editor/GameProjectGeneration/Private/SProjectDialog.cpp
--- a/Engine/Source/Editor/GameProjectGeneration/Private/SProjectDialog.cpp
+++ b/Engine/Source/Editor/GameProjectGeneration/Private/SProjectDialog.cpp
@@ -340,7 +340,7 @@ private:
 		{
 			return ItemPtr->Thumbnail.Get();
 		}
-		return FAppStyle::GetBrush("UnrealDefaultThumbnail");
+		return FAppStyle::GetBrush("UnrealDefaultThumbnailRounded");
 	}
 	
 };
@@ -479,6 +479,15 @@ private:
 	TAttribute<FText> WarningText;
 };
 
+TSharedRef<FSlateBrush> SProjectDialog::SetupThumbnailBrush(TSharedRef<FSlateBrush> ThumbnailBrush)
+{
+	ThumbnailBrush->OutlineSettings.CornerRadii = FVector4(4, 4, 0, 0);
+	ThumbnailBrush->OutlineSettings.RoundingType = ESlateBrushRoundingType::FixedRadius;
+	ThumbnailBrush->DrawAs = ESlateBrushDrawType::RoundedBox;
+
+	return ThumbnailBrush;
+}
+
 void SProjectDialog::Construct(const FArguments& InArgs, EProjectDialogModeMode Mode)
 {
 	bLastGlobalValidityCheckSuccessful = true;
@@ -695,7 +704,7 @@ TSharedRef<SWidget> SProjectDialog::MakeTemplateProjectView()
 			]
 			// Selected template details
 			+ SHorizontalBox::Slot()
-			.Padding(8.0f, 0.0f)
+			.Padding(8.0f, 0.0f, 0.0f, 0.0f)
 			.AutoWidth()
 			[
 				SNew(SScrollBox)
@@ -717,7 +726,7 @@ TSharedRef<SWidget> SProjectDialog::MakeTemplateProjectView()
 						.Padding(FMargin(0.0f, 0.0f, 0.0f, 15.f))
 						[
 							SNew(SImage)
-							.DesiredSizeOverride(FVector2D(358, 160))
+							.DesiredSizeOverride(FVector2D(400, 200))
 							.Image(this, &SProjectDialog::GetSelectedTemplatePreviewImage)
 						]
 						// Template Name
@@ -1225,11 +1234,7 @@ TSharedRef<SWidget> SProjectDialog::MakeRecentProjectsTile()
 
 	static const FName BrushName =  *(FAppStyle::Get().GetContentRootDir() / TEXT("/Starship/Projects/") / TEXT("RecentProjects_2x.png"));
 
-	RecentProjectsBrush = MakeShared<FSlateDynamicImageBrush>(BrushName, FVector2D(300, 100));
-	RecentProjectsBrush->OutlineSettings.CornerRadii = FVector4(4, 4, 4, 4);
-	RecentProjectsBrush->OutlineSettings.RoundingType = ESlateBrushRoundingType::FixedRadius;
-	RecentProjectsBrush->DrawAs = ESlateBrushDrawType::RoundedBox;
-
+	RecentProjectsBrush = SetupThumbnailBrush(MakeShared<FSlateDynamicImageBrush>(BrushName, FVector2D(300, 100)));
 	RecentProjectsCategory->Icon = RecentProjectsBrush;
 
 	return
@@ -1639,7 +1644,7 @@ TMap<FName, TArray<TSharedPtr<FTemplateItem>> > SProjectDialog::FindTemplateProj
 		if (FPlatformFileManager::Get().GetPlatformFile().FileExists(*ThumbnailPNGFile))
 		{
 			const FName BrushName = FName(*ThumbnailPNGFile);
-			Template->Thumbnail = MakeShareable(new FSlateDynamicImageBrush(BrushName, FVector2D(128, 128)));
+			Template->Thumbnail = SetupThumbnailBrush(MakeShared<FSlateDynamicImageBrush>(BrushName, FVector2D(128, 128)));
 		}
 
 		TSharedPtr<FSlateDynamicImageBrush> PreviewBrush;
@@ -1692,7 +1697,7 @@ TMap<FName, TArray<TSharedPtr<FTemplateItem>> > SProjectDialog::FindTemplateProj
 		BlankTemplate->Description = LOCTEXT("BlankProjectDescription", "A clean empty project with no code and default settings.");
 		BlankTemplate->Key = TEXT("Blank");
 		BlankTemplate->SortKey = TEXT("_1");
-		BlankTemplate->Thumbnail = MakeShareable(new FSlateBrush(*FAppStyle::GetBrush("GameProjectDialog.BlankProjectThumbnail")));
+		BlankTemplate->Thumbnail = SetupThumbnailBrush(MakeShared<FSlateBrush>(*FAppStyle::GetBrush("GameProjectDialog.BlankProjectThumbnail")));
 		BlankTemplate->PreviewImage = MakeShareable(new FSlateBrush(*FAppStyle::GetBrush("GameProjectDialog.BlankProjectPreview")));
 		BlankTemplate->BlueprintProjectFile = TEXT("");
 		BlankTemplate->CodeProjectFile = TEXT("");
@@ -2025,10 +2030,7 @@ TArray<TSharedPtr<FTemplateCategory>> SProjectDialog::GetAllTemplateCategories()
 
 		if (!CustomTemplateBrush)
 		{
-			CustomTemplateBrush = MakeShared<FSlateDynamicImageBrush>(BrushName, FVector2D(300, 100));
-			CustomTemplateBrush->OutlineSettings.CornerRadii = FVector4(4, 4, 4, 4);
-			CustomTemplateBrush->OutlineSettings.RoundingType = ESlateBrushRoundingType::FixedRadius;
-			CustomTemplateBrush->DrawAs = ESlateBrushDrawType::RoundedBox;
+			CustomTemplateBrush = SetupThumbnailBrush(MakeShared<FSlateDynamicImageBrush>(BrushName, FVector2D(300, 100)));
 		}
 
 		TSharedPtr<FTemplateCategory> DefaultCategory = MakeShared<FTemplateCategory>();
diff --git a/Engine/Source/Editor/GameProjectGeneration/Private/SProjectDialog.h b/Engine/Source/Editor/GameProjectGeneration/Private/SProjectDialog.h
--- a/Engine/Source/Editor/GameProjectGeneration/Private/SProjectDialog.h
+++ b/Engine/Source/Editor/GameProjectGeneration/Private/SProjectDialog.h
@@ -217,6 +217,7 @@ private:
 
 
 	static TArray<TSharedPtr<FTemplateCategory>> GetAllTemplateCategories();
+	static TSharedRef<FSlateBrush> SetupThumbnailBrush(TSharedRef<FSlateBrush> ThumbnailBrush);
 
 private:
 	TArray<TSharedPtr<FTemplateCategory>> TemplateCategories;
