New Project Window Localized Strings

Ensure title and description are updated if localization changes while the New Project window is opened.

diff --git a/Engine/Source/Editor/GameProjectGeneration/Private/GameProjectGenerationModule.cpp b/Engine/Source/Editor/GameProjectGeneration/Private/GameProjectGenerationModule.cpp
--- a/Engine/Source/Editor/GameProjectGeneration/Private/GameProjectGenerationModule.cpp
+++ b/Engine/Source/Editor/GameProjectGeneration/Private/GameProjectGenerationModule.cpp
@@ -217,8 +217,8 @@ void FGameProjectGenerationModule::LoadTemplateCategories()
 				}
 
 				TemplateCategory->Key = Category.Key;
-				TemplateCategory->DisplayName = FLocalizedTemplateString::GetLocalizedText(Category.LocalizedDisplayNames);
-				TemplateCategory->Description = FLocalizedTemplateString::GetLocalizedText(Category.LocalizedDescriptions);
+				TemplateCategory->LocalizedDisplayNames = Category.LocalizedDisplayNames;
+				TemplateCategory->LocalizedDescriptions = Category.LocalizedDescriptions;
 				
 				const FName BrushName(*Category.Icon);
 
diff --git a/Engine/Source/Editor/GameProjectGeneration/Private/SProjectDialog.cpp b/Engine/Source/Editor/GameProjectGeneration/Private/SProjectDialog.cpp
--- a/Engine/Source/Editor/GameProjectGeneration/Private/SProjectDialog.cpp
+++ b/Engine/Source/Editor/GameProjectGeneration/Private/SProjectDialog.cpp
@@ -93,7 +93,14 @@ public:
 				.TransformPolicy(ETextTransformPolicy::ToUpper)
 				.ShadowOffset(FVector2D(1, 1))
 				.ShadowColorAndOpacity(FLinearColor(0, 0, 0, .75))
-				.Text(Item->DisplayName)
+				.Text_Lambda([WeakItem = Item.ToWeakPtr()]()
+				{
+					if (const TSharedPtr<FTemplateCategory> TemplateCategory = WeakItem.Pin())
+					{
+						return TemplateCategory->GetDisplayName();
+					}
+					return FText::GetEmpty();
+				})
 				.WrapTextAt(250.0f)
 			]
 			+ SOverlay::Slot()
@@ -257,7 +264,14 @@ public:
 								.Font(FAppStyle::Get().GetFontStyle("ProjectBrowser.ProjectTile.Font"))
 								.WrapTextAt(NewProjectDialogDefs::TemplateTileWidth-4.0f)
 								.LineBreakPolicy(FBreakIterator::CreateCamelCaseBreakIterator())
-								.Text(InArgs._Item->Name)
+								.Text_Lambda([this]()
+								{
+									if (const TSharedPtr<FTemplateItem> ItemPtr = Item.Pin())
+									{
+										return ItemPtr->GetDisplayName();
+									}
+									return FText::GetEmpty();
+								})
 								.ColorAndOpacity_Lambda
 								(
 									[this]()
@@ -752,7 +766,14 @@ TSharedRef<SWidget> SProjectDialog::MakeTemplateProjectView()
 							.TextStyle(FAppStyle::Get(), "DialogButtonText")
 							.Font(FAppStyle::Get().GetFontStyle("HeadingExtraSmall"))
 							.ColorAndOpacity(FAppStyle::Get().GetSlateColor("Colors.White"))
-							.Text(this, &SProjectDialog::GetSelectedTemplateProperty, &FTemplateItem::Name)
+							.Text_Lambda([this]()
+							{
+								if (const TSharedPtr<FTemplateItem> SelectedItem = GetSelectedTemplateItem())
+								{
+									return SelectedItem->GetDisplayName();
+								}
+								return FText::GetEmpty();
+							})
 						]
 						// Template Description
 						+ SVerticalBox::Slot()
@@ -772,7 +793,14 @@ TSharedRef<SWidget> SProjectDialog::MakeTemplateProjectView()
 									[
 										SNew(STextBlock)
 										.WrapTextAt(350.0f)
-										.Text(this, &SProjectDialog::GetSelectedTemplateProperty, &FTemplateItem::Description)
+										.Text_Lambda([this]()
+										{
+											if (const TSharedPtr<FTemplateItem> SelectedItem = GetSelectedTemplateItem())
+											{
+												return SelectedItem->GetDescription();
+											}
+											return FText::GetEmpty();
+										})
 									]
 									// Asset types
 									+ SVerticalBox::Slot()
@@ -1240,8 +1268,8 @@ TSharedRef<SWidget> SProjectDialog::MakeRecentProjectsTile()
 {
 	RecentProjectsCategory = MakeShared<FTemplateCategory>();
 
-	RecentProjectsCategory->DisplayName = LOCTEXT("RecentProjects", "Recent Projects");
-	RecentProjectsCategory->Description = FText::GetEmpty();
+	RecentProjectsCategory->DisplayNameOverride = LOCTEXT("RecentProjects", "Recent Projects");
+	RecentProjectsCategory->DescriptionOverride = FText::GetEmpty();
 	RecentProjectsCategory->Key = "RecentProjects";
 	RecentProjectsCategory->IsEnterprise = false;
 
@@ -1638,7 +1666,8 @@ TMap<FName, TArray<TSharedPtr<FTemplateItem>> > SProjectDialog::FindTemplateProj
 
 		// @todo: These are all basically just copies of what's in UTemplateProjectDefs, but ignore differences between code and BP 
 		Template->Categories = TemplateCategoryNames;
-		Template->Description = TemplateDefs->GetLocalizedDescription();
+		Template->LocalizedDisplayNames = TemplateDefs->LocalizedDisplayNames;
+		Template->LocalizedDescriptions = TemplateDefs->LocalizedDescriptions;
 		Template->ClassTypes = TemplateDefs->ClassTypes;
 		Template->AssetTypes = TemplateDefs->AssetTypes;
 		Template->HiddenSettings = TemplateDefs->HiddenSettings;
@@ -1647,12 +1676,6 @@ TMap<FName, TArray<TSharedPtr<FTemplateItem>> > SProjectDialog::FindTemplateProj
 		Template->bThumbnailAsIcon = TemplateDefs->bThumbnailAsIcon;
 		Template->PluginWeak = Plugin;
 
-		Template->Name = TemplateDefs->GetDisplayNameText();
-		if (Template->Name.IsEmpty())
-		{
-			Template->Name = FText::FromString(TemplateKey);
-		}
-
 		const FString ThumbnailPNGFile = (Folder + TEXT("/Media/") + FoundProjectFiles[0]).Replace(TEXT(".uproject"), TEXT(".png"));
 		if (FPlatformFileManager::Get().GetPlatformFile().FileExists(*ThumbnailPNGFile))
 		{
@@ -1706,8 +1729,8 @@ TMap<FName, TArray<TSharedPtr<FTemplateItem>> > SProjectDialog::FindTemplateProj
 	// Add blank template to empty categories
 	{
 		TSharedPtr<FTemplateItem> BlankTemplate = MakeShareable(new FTemplateItem());
-		BlankTemplate->Name = LOCTEXT("BlankProjectName", "Blank");
-		BlankTemplate->Description = LOCTEXT("BlankProjectDescription", "A clean empty project with no code and default settings.");
+		BlankTemplate->DisplayNameOverride = LOCTEXT("BlankProjectName", "Blank");
+		BlankTemplate->DescriptionOverride = LOCTEXT("BlankProjectDescription", "A clean empty project with no code and default settings.");
 		BlankTemplate->Key = TEXT("Blank");
 		BlankTemplate->SortKey = TEXT("_1");
 		BlankTemplate->Thumbnail = SetupThumbnailBrush(MakeShared<FSlateBrush>(*FAppStyle::GetBrush("GameProjectDialog.BlankProjectThumbnail")));
@@ -2061,8 +2084,8 @@ TArray<TSharedPtr<FTemplateCategory>> SProjectDialog::GetAllTemplateCategories()
 
 		TSharedPtr<FTemplateCategory> DefaultCategory = MakeShared<FTemplateCategory>();
 		DefaultCategory->Key = NewProjectDialogDefs::BlankCategoryKey;
-		DefaultCategory->DisplayName = LOCTEXT("ProjectDialog_DefaultCategoryName", "Blank Project");
-		DefaultCategory->Description = LOCTEXT("ProjectDialog_DefaultCategoryDescription", "Create a new blank Unreal project.");
+		DefaultCategory->DisplayNameOverride = LOCTEXT("ProjectDialog_DefaultCategoryName", "Blank Project");
+		DefaultCategory->DescriptionOverride = LOCTEXT("ProjectDialog_DefaultCategoryDescription", "Create a new blank Unreal project.");
 		DefaultCategory->Icon = CustomTemplateBrush;
 
 		AllTemplateCategories.Add(DefaultCategory);
diff --git a/Engine/Source/Editor/GameProjectGeneration/Private/TemplateCategory.h b/Engine/Source/Editor/GameProjectGeneration/Private/TemplateCategory.h
--- a/Engine/Source/Editor/GameProjectGeneration/Private/TemplateCategory.h
+++ b/Engine/Source/Editor/GameProjectGeneration/Private/TemplateCategory.h
@@ -3,16 +3,20 @@
 #pragma once
 
 #include "CoreMinimal.h"
+#include "TemplateProjectDefs.h"
 
 struct FSlateBrush;
 
 struct FTemplateCategory
 {
-	/** Localised name of this category */
-	FText DisplayName;
+	/** The overriden name of this category */
+	FText DisplayNameOverride;
 	
-	/** A description of the templates contained within this category */
-	FText Description;
+	/** The overriden description of the templates contained within this category */
+	FText DescriptionOverride;
+
+	TArray<FLocalizedTemplateString> LocalizedDisplayNames;
+	TArray<FLocalizedTemplateString> LocalizedDescriptions;
 
 	/** A thumbnail to help identify this category (on the tab) */
 	TSharedPtr<FSlateBrush> Icon;
@@ -21,5 +25,24 @@ struct FTemplateCategory
 	FName Key;
 
 	/** Is this an enterprise project? Will end up including Datasmith modules if true. */
-	bool IsEnterprise;
+	bool IsEnterprise = false;
+
+	/** Localized name of this category */
+	FText GetDisplayName() const
+	{
+		if (!DisplayNameOverride.IsEmpty())
+		{
+			return DisplayNameOverride;
+		}
+		return FLocalizedTemplateString::GetLocalizedText(LocalizedDisplayNames);
+	}
+	/** A description of the templates contained within this category */
+	FText GetDescription() const
+	{
+		if (!DescriptionOverride.IsEmpty())
+		{
+			return DescriptionOverride;
+		}
+		return FLocalizedTemplateString::GetLocalizedText(LocalizedDescriptions);
+	}
 };
diff --git a/Engine/Source/Editor/GameProjectGeneration/Private/TemplateItem.h b/Engine/Source/Editor/GameProjectGeneration/Private/TemplateItem.h
--- a/Engine/Source/Editor/GameProjectGeneration/Private/TemplateItem.h
+++ b/Engine/Source/Editor/GameProjectGeneration/Private/TemplateItem.h
@@ -10,8 +10,10 @@
 /** Struct describing a single template project */
 struct FTemplateItem
 {
-	FText		Name;
-	FText		Description;
+	FText		DisplayNameOverride;
+	FText		DescriptionOverride;
+	TArray<FLocalizedTemplateString> LocalizedDisplayNames;
+	TArray<FLocalizedTemplateString> LocalizedDescriptions;
 	TArray<FName> Categories;
 
 	FString		Key;
@@ -36,4 +38,23 @@ struct FTemplateItem
 	bool		bIsBlankTemplate = false;
 	bool		bThumbnailAsIcon = false;
 	
+	/** Localized name of this template */
+	FText GetDisplayName() const
+	{
+		if (!DisplayNameOverride.IsEmpty())
+		{
+			return DisplayNameOverride;
+		}
+		const FText Name = FLocalizedTemplateString::GetLocalizedText(LocalizedDisplayNames);
+		return Name.IsEmpty() ? FText::FromString(Key) : Name;
+	}
+	/** A description of the template */
+	FText GetDescription() const
+	{
+		if (!DescriptionOverride.IsEmpty())
+		{
+			return DescriptionOverride;
+		}
+		return FLocalizedTemplateString::GetLocalizedText(LocalizedDescriptions);
+	}
 };
diff --git a/Engine/Source/Editor/GameProjectGeneration/Private/Tests/GameProjectAutomationTests.cpp b/Engine/Source/Editor/GameProjectGeneration/Private/Tests/GameProjectAutomationTests.cpp
--- a/Engine/Source/Editor/GameProjectGeneration/Private/Tests/GameProjectAutomationTests.cpp
+++ b/Engine/Source/Editor/GameProjectGeneration/Private/Tests/GameProjectAutomationTests.cpp
@@ -149,7 +149,7 @@ namespace GameProjectAutomationUtils
 						if (!GameProjectUtils::CreateProject(ProjectInfo, FailReason, FailLog, &CreatedFiles))
 						{
 							// Failed, report the reason
-							UE_LOG(LogGameProjectGenerationTests, Error, TEXT("Failed to create %s project %s based on %s. Reason: %s\nProject Creation Failure Log:\n%s"), *InCategory.ToString(), *DesiredProjectFilename, *Item->Name.ToString(), *FailReason.ToString(), *FailLog.ToString());
+							UE_LOG(LogGameProjectGenerationTests, Error, TEXT("Failed to create %s project %s based on %s. Reason: %s\nProject Creation Failure Log:\n%s"), *InCategory.ToString(), *DesiredProjectFilename, *Item->GetDisplayName().ToString(), *FailReason.ToString(), *FailLog.ToString());
 						}
 						else
 						{
