PR #14153: Fix formatting of failure reason text

#rnx
#Synced-CL 49240897

diff --git a/Engine/Source/Editor/GameProjectGeneration/Private/GameProjectUtils.cpp b/Engine/Source/Editor/GameProjectGeneration/Private/GameProjectUtils.cpp
--- a/Engine/Source/Editor/GameProjectGeneration/Private/GameProjectUtils.cpp
+++ b/Engine/Source/Editor/GameProjectGeneration/Private/GameProjectUtils.cpp
@@ -1637,7 +1637,7 @@ static bool SaveConfigValues(const FProjectInformation& InProjectInfo, const TAr
 			{
 				FFormatNamedArguments Args;
 				Args.Add(TEXT("ConfigFile"), FText::FromString(IniFilename));
-				OutFailReason = LOCTEXT("FailedToFixUpConfigFile", "Failed to process file {ConfigFile}.");
+				OutFailReason = FText::Format(LOCTEXT("FailedToFixUpConfigFile", "Failed to process file {ConfigFile}."), Args);
 				return false;
 			}
 		}
