Revamped New Project Dialog's template description section

[Removed] the ScrollBar on the description section.
[Moved] Plugin information panel to the description section to solve issue with newly added scroll bar which forced the template row to display 3 tiles instead of 4.

#Synced-CL 50650077

diff --git a/Engine/Source/Editor/GameProjectGeneration/Private/SProjectDialog.cpp b/Engine/Source/Editor/GameProjectGeneration/Private/SProjectDialog.cpp
--- a/Engine/Source/Editor/GameProjectGeneration/Private/SProjectDialog.cpp
+++ b/Engine/Source/Editor/GameProjectGeneration/Private/SProjectDialog.cpp
@@ -734,130 +734,127 @@ TSharedRef<SWidget> SProjectDialog::MakeTemplateProjectView()
 			.Padding(8.0f, 0.0f, 0.0f, 0.0f)
 			.AutoWidth()
 			[
-				SNew(SScrollBox)
-				.Orientation(Orient_Vertical)
-				.ConsumeMouseWheel(EConsumeMouseWheel::WhenScrollingPossible)
-				+ SScrollBox::Slot()
-				.FillSize(1.f)
+				SNew(SVerticalBox)
+				+ SVerticalBox::Slot()
+				.FillHeight(1.f)
+				.VAlign(VAlign_Top)
 				[
+					// Preview image
 					SNew(SVerticalBox)
-					.Clipping(EWidgetClipping::ClipToBounds)
 					+ SVerticalBox::Slot()
-					.FillHeight(1.f)
+					.AutoHeight()
 					.Padding(FMargin(0.0f, 0.0f, 0.0f, 15.f))
 					[
-						// Preview image
-						SNew(SVerticalBox)
-						+ SVerticalBox::Slot()
-						.AutoHeight()
-						.Padding(FMargin(0.0f, 0.0f, 0.0f, 15.f))
-						[
-							SNew(SImage)
-							.DesiredSizeOverride(FVector2D(400, 200))
-							.Image(this, &SProjectDialog::GetSelectedTemplatePreviewImage)
-						]
-						// Template Name
-						+ SVerticalBox::Slot()
-						.Padding(FMargin(0.0f, 0.0f, 0.0f, 10.0f))
-						.AutoHeight()
-						[
-							SNew(STextBlock)
-							.AutoWrapText(true)
-							.TextStyle(FAppStyle::Get(), "DialogButtonText")
-							.Font(FAppStyle::Get().GetFontStyle("HeadingExtraSmall"))
-							.ColorAndOpacity(FAppStyle::Get().GetSlateColor("Colors.White"))
-							.Text_Lambda([this]()
+						SNew(SImage)
+						.DesiredSizeOverride(FVector2D(400, 200))
+						.Image(this, &SProjectDialog::GetSelectedTemplatePreviewImage)
+					]
+					// Template Name
+					+ SVerticalBox::Slot()
+					.Padding(FMargin(0.0f, 0.0f, 0.0f, 10.0f))
+					.AutoHeight()
+					[
+						SNew(STextBlock)
+						.AutoWrapText(true)
+						.TextStyle(FAppStyle::Get(), "DialogButtonText")
+						.Font(FAppStyle::Get().GetFontStyle("HeadingExtraSmall"))
+						.ColorAndOpacity(FAppStyle::Get().GetSlateColor("Colors.White"))
+						.Text_Lambda([this]()
+						{
+							if (const TSharedPtr<FTemplateItem> SelectedItem = GetSelectedTemplateItem())
 							{
-								if (const TSharedPtr<FTemplateItem> SelectedItem = GetSelectedTemplateItem())
-								{
-									return SelectedItem->GetDisplayName();
-								}
-								return FText::GetEmpty();
-							})
-						]
-						// Template Description
-						+ SVerticalBox::Slot()
-						.AutoHeight()
+								return SelectedItem->GetDisplayName();
+							}
+							return FText::GetEmpty();
+						})
+					]
+					// Template Description
+					+ SVerticalBox::Slot()
+					.FillHeight(1.f)
+					[
+						SNew(SBox)
+						.WidthOverride(358.0f)
 						[
-							SNew(SBox)
-							.HeightOverride(120.0f)
-							.WidthOverride(358.0f)
-							[
-								SNew(SScrollBox)
-								+ SScrollBox::Slot()
-								[ 
+							SNew(SScrollBox)
+							+ SScrollBox::Slot()
+							[ 
+								SNew(SVerticalBox)
+								+ SVerticalBox::Slot()
+								.Padding(FMargin(0.0f, 0.0f, 0.0f, 10.0f))
+								.AutoHeight()
+								[
+									SNew(STextBlock)
+									.WrapTextAt(350.0f)
+									.Text_Lambda([this]()
+									{
+										if (const TSharedPtr<FTemplateItem> SelectedItem = GetSelectedTemplateItem())
+										{
+											return SelectedItem->GetDescription();
+										}
+										return FText::GetEmpty();
+									})
+								]
+								// Asset types
+								+ SVerticalBox::Slot()
+								.AutoHeight()
+								.Padding(FMargin(0.0f, 5.0f, 0.0f, 5.0f))
+								[
 									SNew(SVerticalBox)
+									.Visibility(this, &SProjectDialog::GetSelectedTemplateAssetVisibility)
 									+ SVerticalBox::Slot()
-									.Padding(FMargin(0.0f, 0.0f, 0.0f, 10.0f))
-									.AutoHeight()
 									[
 										SNew(STextBlock)
-										.WrapTextAt(350.0f)
-										.Text_Lambda([this]()
-										{
-											if (const TSharedPtr<FTemplateItem> SelectedItem = GetSelectedTemplateItem())
-											{
-												return SelectedItem->GetDescription();
-											}
-											return FText::GetEmpty();
-										})
+										.TextStyle(FAppStyle::Get(), "DialogButtonText")
+										.Font(FAppStyle::Get().GetFontStyle("NormalFontBold"))
+										.ColorAndOpacity(FAppStyle::Get().GetSlateColor("Colors.White"))
+										.Text(LOCTEXT("ProjectTemplateAssetTypes", "Asset Type References"))
 									]
-									// Asset types
 									+ SVerticalBox::Slot()
 									.AutoHeight()
-									.Padding(FMargin(0.0f, 5.0f, 0.0f, 5.0f))
 									[
-										SNew(SVerticalBox)
-										.Visibility(this, &SProjectDialog::GetSelectedTemplateAssetVisibility)
-										+ SVerticalBox::Slot()
-										[
-											SNew(STextBlock)
-											.TextStyle(FAppStyle::Get(), "DialogButtonText")
-											.Text(LOCTEXT("ProjectTemplateAssetTypes", "Asset Type References"))
-										]
-										+ SVerticalBox::Slot()
-										.AutoHeight()
-										[
-											SNew(STextBlock)
-											.AutoWrapText(true)
-											.Text(this, &SProjectDialog::GetSelectedTemplateAssetTypes)
-										]
+										SNew(STextBlock)
+										.AutoWrapText(true)
+										.Text(this, &SProjectDialog::GetSelectedTemplateAssetTypes)
+									]
+								]
+								// Class types
+								+ SVerticalBox::Slot()
+								.AutoHeight()
+								.Padding(FMargin(0.0f, 5.0f, 0.0f, 5.0f))
+								[
+									SNew(SVerticalBox)
+									.Visibility(this, &SProjectDialog::GetSelectedTemplateClassVisibility)
+									+ SVerticalBox::Slot()
+									[
+										SNew(STextBlock)
+										.TextStyle(FAppStyle::Get(), "DialogButtonText")
+										.Font(FAppStyle::Get().GetFontStyle("NormalFontBold"))
+										.ColorAndOpacity(FAppStyle::Get().GetSlateColor("Colors.White"))
+										.Text(LOCTEXT("ProjectTemplateClassTypes", "Class Type References"))
 									]
-									// Class types
 									+ SVerticalBox::Slot()
 									.AutoHeight()
-									.Padding(FMargin(0.0f, 5.0f, 0.0f, 5.0f))
 									[
-										SNew(SVerticalBox)
-										.Visibility(this, &SProjectDialog::GetSelectedTemplateClassVisibility)
-										+ SVerticalBox::Slot()
-										[
-											SNew(STextBlock)
-											.TextStyle(FAppStyle::Get(), "DialogButtonText")
-											.Text(LOCTEXT("ProjectTemplateClassTypes", "Class Type References"))
-										]
-										+ SVerticalBox::Slot()
-										.AutoHeight()
-										[
-											SNew(STextBlock)
-											.AutoWrapText(true)
-											.Text(this, &SProjectDialog::GetSelectedTemplateClassTypes)
-										]
+										SNew(STextBlock)
+										.AutoWrapText(true)
+										.Text(this, &SProjectDialog::GetSelectedTemplateClassTypes)
 									]
 								]
+								// Plugin Information
+								+ SVerticalBox::Slot()
+								.AutoHeight()
+								.Padding(FMargin(0.0f, 5.0f, 0.0f, 0.0f))
+								.Expose(PluginInformationSlot)
 							]
 						]
-						// Project Options
-						+ SVerticalBox::Slot()
-						.Padding(FMargin(0.0f, 0.0f, 0.0f, 15.0f))
-						.AutoHeight()
-						.Expose(ProjectOptionsSlot)
-						// Plugin Information
-						+ SVerticalBox::Slot()
-						.AutoHeight()
-						.Expose(PluginInformationSlot)
 					]
 				]
+				// Project Options
+				+ SVerticalBox::Slot()
+				.Padding(FMargin(0.0f, 0.0f, 0.0f, 15.0f))
+				.AutoHeight()
+				.VAlign(VAlign_Bottom)
+				.Expose(ProjectOptionsSlot)
 			]
 		];
 }
@@ -1183,6 +1180,10 @@ TSharedRef<SWidget> SProjectDialog::MakePluginInformationWidget()
 		PluginIconDynamicImageBrush = MakeShareable(new FSlateDynamicImageBrush(BrushName, FVector2D(Size.X, Size.Y)));
 	}
 
+	const FText PluginName = FText::FromString(Plugin->GetFriendlyName());
+	const FText PluginAuthor = FText::FromString(PluginDescriptor.CreatedBy);
+
+
 	TSharedRef<SWidget> PluginInformationWidget =
 		SNew(SVerticalBox)
 		.Visibility(this, &SProjectDialog::GetProjectSettingsVisibility)
@@ -1190,14 +1191,11 @@ TSharedRef<SWidget> SProjectDialog::MakePluginInformationWidget()
 		+ SVerticalBox::Slot()
 		.AutoHeight()
 		[
-			SNew(SBorder)
-			.Padding(FMargin(10.0f, 7.0f))
-			.BorderImage(FAppStyle::Get().GetBrush("Brushes.Header"))
-			[
-				SNew(STextBlock)
-				.TextStyle(FAppStyle::Get(), "DialogButtonText")
-				.Text(LOCTEXT("PluginInformation", "Plugin Information"))
-			]
+			SNew(STextBlock)
+			.TextStyle(FAppStyle::Get(), "DialogButtonText")
+			.Font(FAppStyle::Get().GetFontStyle("NormalFontBold"))
+			.ColorAndOpacity(FAppStyle::Get().GetSlateColor("Colors.White"))
+			.Text(LOCTEXT("PluginInformation", "Plugin Information"))
 		]
 		// Plugin tile 
 		+ SVerticalBox::Slot()
@@ -1212,48 +1210,73 @@ TSharedRef<SWidget> SProjectDialog::MakePluginInformationWidget()
 				.BorderImage(FAppStyle::Get().GetBrush("Brushes.Panel"))
 				.Padding(0)
 				[
-					SNew(SBorder)
-					.BorderImage(PluginStyle->GetBrush("PluginTile.BorderImage"))
-					.Padding(PaddingAmount + 4.f)
+					SNew(SHorizontalBox)
+					// Thumbnail image
+					+ SHorizontalBox::Slot()
+					.Padding(PaddingAmount, PaddingAmount, PaddingAmount + 10.f, PaddingAmount)
+					.VAlign(VAlign_Center)
+					.AutoWidth()
 					[
-						SNew(SHorizontalBox)
-						// Thumbnail image
-						+ SHorizontalBox::Slot()
-						.Padding(PaddingAmount, PaddingAmount, PaddingAmount + 10.f, PaddingAmount)
-						.VAlign(VAlign_Center)
-						.AutoWidth()
+						SNew(SBox)
+						.WidthOverride(ThumbnailImageSize)
+						.HeightOverride(ThumbnailImageSize)
 						[
-							SNew(SBox)
-							.WidthOverride(ThumbnailImageSize)
-							.HeightOverride(ThumbnailImageSize)
+							SNew(SBorder)
+							.BorderImage(PluginStyle->GetBrush("PluginTile.ThumbnailBorderImage"))
 							[
-								SNew(SBorder)
-								.BorderImage(PluginStyle->GetBrush("PluginTile.ThumbnailBorderImage"))
-								[
-									SNew(SImage)
-									.Image(PluginIconDynamicImageBrush.IsValid() ? PluginIconDynamicImageBrush.Get() : PluginStyle->GetBrush("PluginTile.ThumbnailImageFallBack"))
-								]
+								SNew(SImage)
+								.Image(PluginIconDynamicImageBrush.IsValid() ? PluginIconDynamicImageBrush.Get() : PluginStyle->GetBrush("PluginTile.ThumbnailImageFallBack"))
 							]
 						]
-						+ SHorizontalBox::Slot()
+					]
+					+ SHorizontalBox::Slot()
+					.VAlign(VAlign_Center)
+					[
+
+						SNew(SVerticalBox)
+						+ SVerticalBox::Slot()
+						.VAlign(VAlign_Top)
 						[
-							SNew(SVerticalBox)
+							SNew(SHorizontalBox)
 							// Friendly Name
-							+ SVerticalBox::Slot()
-							.AutoHeight()
-							.Padding(PaddingAmount, PaddingAmount + 3.f, PaddingAmount + 8.f, 0.f)
+							+SHorizontalBox::Slot()
+							.VAlign(VAlign_Center)
+							.HAlign(HAlign_Left)
 							[
 								SNew(STextBlock)
-								.Text(FText::FromString(Plugin->GetFriendlyName()))
-								.TextStyle(PluginStyle, "PluginTile.NameText")
+								.Text(PluginName)
+								.OverflowPolicy(ETextOverflowPolicy::Ellipsis)
+								.ToolTipText(PluginName)
+								.TextStyle(FAppStyle::Get(), "DialogButtonText")
+								.ColorAndOpacity(FAppStyle::Get().GetSlateColor("Colors.White"))
 							]
 							// Created By
-							+ SVerticalBox::Slot()
-							.AutoHeight()
-							.Padding(PaddingAmount, PaddingAmount, 0, 0)
+							+ SHorizontalBox::Slot()
+							.VAlign(VAlign_Center)
+							.HAlign(HAlign_Right)
+							[
+								SNew(STextBlock)
+								.Text(PluginAuthor)
+								.TextStyle(FAppStyle::Get(), "DialogButtonText")
+							]
+						]
+						// Version
+						+ SVerticalBox::Slot()
+						.AutoHeight()
+						.VAlign(VAlign_Bottom)
+						.HAlign(HAlign_Left)
+						.Padding(0.0f, PaddingAmount, 0.0f, 0.0f)
+						[
+							SNew(SHorizontalBox)
+							// version number
+							+ SHorizontalBox::Slot()
+							.AutoWidth()
+							.VAlign(VAlign_Center)
+							.Padding(0.0f, 0.0f, 0.0f, 1.0f) // Lower padding to align font with version number base
 							[
 								SNew(STextBlock)
-								.Text(FText::FromString(PluginDescriptor.CreatedBy))
+								.Text(FText::Format(LOCTEXT("PluginVersionLabel", "Version {0}"), FText::FromString(PluginDescriptor.VersionName)))
+								.TextStyle(FAppStyle::Get(), "DialogButtonText")
 							]
 						]
 					]
diff --git a/Engine/Source/Editor/MainFrame/Private/Frame/MainFrameActions.cpp b/Engine/Source/Editor/MainFrame/Private/Frame/MainFrameActions.cpp
--- a/Engine/Source/Editor/MainFrame/Private/Frame/MainFrameActions.cpp
+++ b/Engine/Source/Editor/MainFrame/Private/Frame/MainFrameActions.cpp
@@ -848,12 +848,17 @@ TSharedRef<SWindow> FMainFrameActionCallbacks::CreateProjectBrowserWindow(bool b
 		Title = LOCTEXT("NewProjectWindowHeader", "New Project");
 	}
 
+	constexpr float MinWidthToAvoidOverlap = 900.f;
+	constexpr float MinHeightToAvoidOverlap = 640.f;
+
 	return SNew(SWindow)
 		.Title(Title)
 		.ClientSize(FMainFrameModule::GetProjectBrowserWindowSize())
 		.SizingRule(ESizingRule::UserSized)
 		.SupportsMinimize(false)
-		.SupportsMaximize(false);
+		.SupportsMaximize(false)
+		.MinWidth(MinWidthToAvoidOverlap)
+		.MinHeight(MinHeightToAvoidOverlap);
 }
 
 
diff --git a/Engine/Source/Editor/MainFrame/Private/MainFrameModule.cpp b/Engine/Source/Editor/MainFrame/Private/MainFrameModule.cpp
--- a/Engine/Source/Editor/MainFrame/Private/MainFrameModule.cpp
+++ b/Engine/Source/Editor/MainFrame/Private/MainFrameModule.cpp
@@ -424,6 +424,9 @@ void FMainFrameModule::CreateDefaultMainFrameAuxiliary(const bool bStartImmersiv
 		TSharedPtr<SWindow> RootWindowPtr = GetParentWindow();
 		if (!RootWindowPtr.IsValid())
 		{
+			constexpr float MinWidthToAvoidOverlap = 900.f;
+			constexpr float MinHeightToAvoidOverlap = 640.f;
+
 			TSharedRef<SWindow> RootWindow = SNew(SWindow)
 				.AutoCenter(WindowConfig.CenterRules)
 				.Title(WindowConfig.WindowTitle)
@@ -437,7 +440,9 @@ void FMainFrameModule::CreateDefaultMainFrameAuxiliary(const bool bStartImmersiv
 				// this is needed to allow Slate-rendered Metal views to show up on XCode's view debugger
 				.SupportsTransparency(EWindowTransparency::PerPixel)
 #endif
-				.SupportsMinimize(WindowConfig.bSupportsMinimize);
+				.SupportsMinimize(WindowConfig.bSupportsMinimize)
+				.MinWidth(MinWidthToAvoidOverlap)
+				.MinHeight(MinHeightToAvoidOverlap);
 
 			// Avoid using the is initially maximized argument when creating the window so that the code below can change the window without this initial value taking over taking the changes when the window is first shown.
 			if (DefaultWindowLocation.InitiallyMaximized)
