[HomeScreen] Deprecate CreateProjectBrowser function
Deprecate the old CreateProjectBrowser function in favor of a different one that should create just the embedded project browser.

#Synced-CL 50450741

diff --git a/Engine/Source/Editor/GameProjectGeneration/Private/GameProjectGenerationModule.cpp b/Engine/Source/Editor/GameProjectGeneration/Private/GameProjectGenerationModule.cpp
--- a/Engine/Source/Editor/GameProjectGeneration/Private/GameProjectGenerationModule.cpp
+++ b/Engine/Source/Editor/GameProjectGeneration/Private/GameProjectGenerationModule.cpp
@@ -73,6 +73,11 @@ TSharedRef<SWidget> FGameProjectGenerationModule::CreateProjectBrowser(bool bInS
 	return SNew(SProjectBrowser).ShowHeaderBar(bInShowHeaderBar).AllowScrolling(bInAllowScrolling);
 }
 
+TSharedRef<class SWidget> FGameProjectGenerationModule::CreateEmbeddedProjectBrowser()
+{
+	return SNew(SProjectBrowser).IsEmbedded(true);
+}
+
 void FGameProjectGenerationModule::OpenAddCodeToProjectDialog(const FAddToProjectConfig& Config)
 {
 	GameProjectUtils::OpenAddToProjectDialog(Config, EClassDomain::Native);
diff --git a/Engine/Source/Editor/GameProjectGeneration/Private/SProjectBrowser.cpp b/Engine/Source/Editor/GameProjectGeneration/Private/SProjectBrowser.cpp
--- a/Engine/Source/Editor/GameProjectGeneration/Private/SProjectBrowser.cpp
+++ b/Engine/Source/Editor/GameProjectGeneration/Private/SProjectBrowser.cpp
@@ -369,6 +369,11 @@ void SProjectBrowser::Construct( const FArguments& InArgs )
 	// Prepare the projects box
 	ProjectsBox = SNew(SVerticalBox);
 
+	// If embedded the HeaderBar and AllowScrolling should not be allowed
+	bIsEmbedded = InArgs._IsEmbedded;
+	const bool ShowHeaderBar = InArgs._ShowHeaderBar && !bIsEmbedded;
+	const bool AllowScrolling = InArgs._AllowScrolling && !bIsEmbedded;
+
 	SAssignNew(ProjectTileView, STileView<TSharedPtr<FProjectItem>>)
 		.ListItemsSource(&FilteredProjectItemsSource)
 		.SelectionMode(ESelectionMode::Single)
@@ -378,10 +383,10 @@ void SProjectBrowser::Construct( const FArguments& InArgs )
 		.OnContextMenuOpening(this, &SProjectBrowser::OnGetContextMenuContent)
 		.OnMouseButtonDoubleClick(this, &SProjectBrowser::HandleProjectItemDoubleClick)
 		.OnSelectionChanged(this, &SProjectBrowser::HandleProjectViewSelectionChanged)
-		.AllowOverscroll(InArgs._AllowScrolling ? EAllowOverscroll::Yes : EAllowOverscroll::No)
-		.ConsumeMouseWheel(InArgs._AllowScrolling ? EConsumeMouseWheel::Always : EConsumeMouseWheel::Never)
-		.ScrollbarVisibility(InArgs._AllowScrolling ? EVisibility::Visible : EVisibility::Collapsed)
-		.WheelScrollMultiplier(InArgs._AllowScrolling ? GetGlobalScrollAmount() : 0.f)
+		.AllowOverscroll(AllowScrolling ? EAllowOverscroll::Yes : EAllowOverscroll::No)
+		.ConsumeMouseWheel(AllowScrolling ? EConsumeMouseWheel::Always : EConsumeMouseWheel::Never)
+		.ScrollbarVisibility(AllowScrolling ? EVisibility::Visible : EVisibility::Collapsed)
+		.WheelScrollMultiplier(AllowScrolling ? GetGlobalScrollAmount() : 0.f)
 		.ItemHeight(ProjectBrowserDefs::ProjectTileHeight+9)
 		.ItemWidth(ProjectBrowserDefs::ProjectTileWidth+9);
 
@@ -435,7 +440,7 @@ void SProjectBrowser::Construct( const FArguments& InArgs )
 		.AutoHeight()
 		[
 			SNew(SHorizontalBox)
-			.Visibility(InArgs._ShowHeaderBar ? EVisibility::Visible : EVisibility::Collapsed)
+			.Visibility(ShowHeaderBar ? EVisibility::Visible : EVisibility::Collapsed)
 
 			+ SHorizontalBox::Slot()
 			.Padding(FMargin(0, 0, 5.f, 5))
diff --git a/Engine/Source/Editor/GameProjectGeneration/Private/SProjectBrowser.h b/Engine/Source/Editor/GameProjectGeneration/Private/SProjectBrowser.h
--- a/Engine/Source/Editor/GameProjectGeneration/Private/SProjectBrowser.h
+++ b/Engine/Source/Editor/GameProjectGeneration/Private/SProjectBrowser.h
@@ -64,9 +64,13 @@ public:
 	SLATE_BEGIN_ARGS(SProjectBrowser)
 		: _ShowHeaderBar(true)
 		, _AllowScrolling(true)
+		, _IsEmbedded(false)
 	{}
 		SLATE_ARGUMENT(bool, ShowHeaderBar)
 		SLATE_ARGUMENT(bool, AllowScrolling)
+
+		/** If true it will show the embedded version of the ProjectBrowser */
+		SLATE_ARGUMENT(bool, IsEmbedded)
 		SLATE_ARGUMENT(FProjectSelectionChanged, OnSelectionChanged);
 	SLATE_END_ARGS()
 
@@ -190,6 +194,9 @@ private:
 
 	bool IsOnlineContentFinished;
 
+	/** If true it will show the embedded version of the ProjectBrowser */
+	bool bIsEmbedded;
+
 	// Holds a delegate that is executed when the new project screen is being requested.
 	FNewProjectScreenRequested NewProjectScreenRequestedDelegate;
 
diff --git a/Engine/Source/Editor/GameProjectGeneration/Public/GameProjectGenerationModule.h b/Engine/Source/Editor/GameProjectGeneration/Public/GameProjectGenerationModule.h
--- a/Engine/Source/Editor/GameProjectGeneration/Public/GameProjectGenerationModule.h
+++ b/Engine/Source/Editor/GameProjectGeneration/Public/GameProjectGenerationModule.h
@@ -61,8 +61,12 @@ public:
 	virtual TSharedRef<class SWidget> CreateNewClassDialog(const UClass* InClass);
 
 	/** Creates the project browser widget */
+	UE_DEPRECATED(5.8, "Use CreateEmbeddedProjectBrowser")
 	virtual TSharedRef<class SWidget> CreateProjectBrowser(bool bInShowHeaderBar = true, bool bInAllowScrolling = true);
 
+	/** Creates the embedded project browser widget */
+	virtual TSharedRef<class SWidget> CreateEmbeddedProjectBrowser();
+
 	/** 
 	 * Opens a dialog to add code files to the current project. 
 	 *
diff --git a/Engine/Source/Editor/MainFrame/Private/MainFrameModule.cpp b/Engine/Source/Editor/MainFrame/Private/MainFrameModule.cpp
--- a/Engine/Source/Editor/MainFrame/Private/MainFrameModule.cpp
+++ b/Engine/Source/Editor/MainFrame/Private/MainFrameModule.cpp
@@ -196,7 +196,6 @@ namespace HomeScreenHelper
 
 	TSharedRef<SWidget> GetRecentProjectListWidget()
 	{
-		constexpr bool bIsEmbedded = true;
 		return SNew(SBox)
 			.Padding(8.f)
 			[
@@ -205,7 +204,7 @@ namespace HomeScreenHelper
 				[
 					// Embedded should not show the HeaderBar and should not allow scrolling
 					FGameProjectGenerationModule::Get()
-						.CreateProjectBrowser(!bIsEmbedded /** ShowHeaderBar */, !bIsEmbedded /** AllowScrolling */)
+						.CreateEmbeddedProjectBrowser()
 				]
 			];
 	}
